AutomatedPerfTest controller hook-up of CitySample.

- First pass, still not yet active.

diff --git a/Engine/Plugins/Performance/AutomatedPerfTesting/Build/Scripts/AutomatedPerfTestConfig.cs b/Engine/Plugins/Performance/AutomatedPerfTesting/Build/Scripts/AutomatedPerfTestConfig.cs
--- a/Engine/Plugins/Performance/AutomatedPerfTesting/Build/Scripts/AutomatedPerfTestConfig.cs
+++ b/Engine/Plugins/Performance/AutomatedPerfTesting/Build/Scripts/AutomatedPerfTestConfig.cs
@@ -11,6 +11,12 @@ namespace AutomatedPerfTest
 {
 	public class AutomatedPerfTestConfigBase : UnrealTestConfiguration
 	{
+		/// <summary>
+		/// Name of the project (TODO: parse from the build metadata? What about github builds?)
+		/// </summary>
+		[AutoParamWithNames("AutomatedPerfTest.ProjectName")]
+		public string ProjectName = "";
+
 		/// <summary>
 		/// If we're running on the build machine
 		/// </summary>
@@ -58,7 +64,13 @@ namespace AutomatedPerfTest
 		/// </summary>
 		[AutoParamWithNames("", "AutomatedPerfTest.PerfCacheRoot")]
 		public string PerfCacheRoot;
-		
+
+		/// <summary>
+		/// Path to a JSON file with ignored issues (ensures, warnings, errros). Can be used to suppress hard-to-fix issues, on a per-branch basis
+		/// </summary>
+		[AutoParamWithNames("", "AutomatedPerfTest.IgnoredIssuesConfigAbsPath")]
+		public string IgnoredIssuesConfigAbsPath;
+
 		/// <summary>
 		/// Which platform we're testing on 
 		/// </summary>
diff --git a/Engine/Plugins/Performance/AutomatedPerfTesting/Build/Scripts/AutomatedPerfTestNode.cs b/Engine/Plugins/Performance/AutomatedPerfTesting/Build/Scripts/AutomatedPerfTestNode.cs
new file mode 100644
--- /dev/null
+++ b/Engine/Plugins/Performance/AutomatedPerfTesting/Build/Scripts/AutomatedPerfTestNode.cs
@@ -0,0 +1,488 @@
+// Copyright Epic Games, Inc. All Rights Reserved.
+
+using System;
+using System.Collections.Generic;
+using System.IO;
+using System.Linq;
+using EpicGame;
+using AutomationTool;
+using Gauntlet;
+using EpicGames.Core;
+using Log = Gauntlet.Log;
+using Microsoft.Extensions.Logging;
+using UnrealBuildBase;	// for Unreal.RootDirectory
+using UnrealBuildTool;	// for UnrealTargetPlatform
+
+using static AutomationTool.CommandUtils;
+
+namespace AutomatedPerfTest
+{
+	/// <summary>
+	/// Implementation of a Gauntlet TestNode for AutomatedPerfTest plugin
+	/// </summary>
+	/// <typeparam name="TConfigClass"></typeparam>
+	public abstract class AutomatedPerfTestNode<TConfigClass> : UnrealTestNode<TConfigClass>
+		where TConfigClass : AutomatedPerfTestConfigBase, new()
+	{
+		public AutomatedPerfTestNode(UnrealTestContext InContext) : base(InContext)
+		{
+			TestGuid = Guid.NewGuid();
+			Gauntlet.Log.Info("Your Test GUID is :\n" + TestGuid.ToString() + '\n');
+
+			InitHandledErrors();
+
+			LastLogLineCount = 0;
+		}
+
+		public override bool StartTest(int Pass, int InNumPasses)
+		{
+			LastLogLineCount = 0;
+			return base.StartTest(Pass, InNumPasses);
+		}
+
+		public class HandledError
+		{
+			public string ClientErrorString;
+			public string GauntletErrorString;
+
+			/// <summary>
+			/// String name for the log category that should be used to filter errors. Defaults to null, i.e. no filter.
+			/// </summary>
+			public string CategoryName;
+
+			// If error is verbose, will output debugging information such as state
+			public bool Verbose;
+
+			public HandledError(string ClientError, string GauntletError, string Category, bool VerboseIn = false)
+			{
+				ClientErrorString = ClientError;
+				GauntletErrorString = GauntletError;
+				CategoryName = Category;
+				Verbose = VerboseIn;
+			}
+		}
+
+		/// <summary>
+		/// List of errors with special-cased gauntlet messages.
+		/// </summary>
+		public List<HandledError> HandledErrors { get; set; }
+
+		/// <summary>
+		/// Guid associated with each test run for ease of differentiation between different runs on same build.
+		/// </summary>
+		public Guid TestGuid { get; protected set; }
+
+		/// <summary>
+		/// Line count of the client log messages that have been written to the test logs.
+		/// </summary>
+		private int LastLogLineCount;
+
+		/// <summary>
+		// Temporary directory for perf report CSVs
+		/// </summary>
+		private DirectoryInfo TempPerfCSVDir => new DirectoryInfo(Path.Combine(Unreal.RootDirectory.FullName, "GauntletTemp", "PerfReportCSVs"));
+
+		/// <summary>
+		// Holds the build name as is, since if this is a preflight the suffix will be stripped after GetConfiguration is called.
+		/// </summary>
+		private string OriginalBuildName = null;
+
+		/// <summary>
+		/// Set up the base list of possible expected errors, plus the messages to deliver if encountered.
+		/// </summary>
+		protected virtual void InitHandledErrors()
+		{
+			HandledErrors = new List<HandledError>();
+		}
+
+		/// <summary>
+		/// Periodically called while test is running. Updates logs.
+		/// </summary>
+		public override void TickTest()
+		{
+			IAppInstance App = null;
+
+			if (TestInstance.ClientApps == null)
+			{
+				App = TestInstance.ServerApp;
+			}
+			else
+			{
+				if (TestInstance.ClientApps.Length > 0)
+				{
+					App = TestInstance.ClientApps.First();
+				}
+			}
+
+			if (App != null)
+			{
+				UnrealLogParser Parser = new UnrealLogParser(App.StdOut);
+
+				string LogChannelName = GetCachedConfiguration().ProjectName + "Test";
+				List<string> TestLines = Parser.GetLogChannel(LogChannelName).ToList();
+
+				string LogCategory = "Log" + LogChannelName;
+				string LogCategoryError = LogCategory + ": Error:";
+				string LogCategoryWarning = LogCategory + ": Warning:";
+				
+				for (int i = LastLogLineCount; i < TestLines.Count; i++)
+				{
+					if (TestLines[i].StartsWith(LogCategoryError))
+					{
+						ReportError(TestLines[i]);
+					}
+					else if (TestLines[i].StartsWith(LogCategoryWarning))
+					{
+						ReportWarning(TestLines[i]);
+					}
+					else
+					{
+						Log.Info(TestLines[i]);
+					}
+				}
+
+				LastLogLineCount = TestLines.Count;
+			}
+
+			base.TickTest();
+		}
+
+		/// <summary>
+		/// This allows using a per-branch config to ignore certain issues
+		/// that were inherited from Main and will be addressed there
+		/// </summary>
+		/// <param name="InArtifacts"></param>
+		/// <returns></returns>
+		protected override UnrealLog CreateLogSummaryFromArtifact(UnrealRoleArtifacts InArtifacts)
+		{
+			UnrealLog LogSummary = base.CreateLogSummaryFromArtifact(InArtifacts);
+
+			IgnoredIssueConfig IgnoredIssues = new IgnoredIssueConfig();
+
+			string IgnoredIssuePath = GetCachedConfiguration().IgnoredIssuesConfigAbsPath;
+
+			if (!File.Exists(IgnoredIssuePath))
+			{
+				Log.Info("No IgnoredIssue Config found at {0}", IgnoredIssuePath);
+			}
+			else if (IgnoredIssues.LoadFromFile(IgnoredIssuePath))
+			{
+				Log.Info("Loaded IgnoredIssue config from {0}", IgnoredIssuePath);
+
+				IEnumerable<UnrealLog.CallstackMessage> IgnoredEnsures = LogSummary.Ensures.Where(E => IgnoredIssues.IsEnsureIgnored(this.Name, E.Message));
+				IEnumerable<UnrealLog.LogEntry> IgnoredWarnings = LogSummary.LogEntries.Where(E => E.Level == UnrealLog.LogLevel.Warning && IgnoredIssues.IsWarningIgnored(this.Name, E.Message));
+				IEnumerable<UnrealLog.LogEntry> IgnoredErrors = LogSummary.LogEntries.Where(E => E.Level == UnrealLog.LogLevel.Error && IgnoredIssues.IsErrorIgnored(this.Name, E.Message));
+
+				if (IgnoredEnsures.Any())
+				{
+					Log.Info("Ignoring {0} ensures.", IgnoredEnsures.Count());
+					Log.Info("\t{0}", string.Join("\n\t", IgnoredEnsures.Select(E => E.Message)));
+					LogSummary.Ensures = LogSummary.Ensures.Except(IgnoredEnsures).ToArray();
+				}
+				if (IgnoredWarnings.Any())
+				{
+					Log.Info("Ignoring {0} warnings.", IgnoredWarnings.Count());
+					Log.Info("\t{0}", string.Join("\n\t", IgnoredWarnings.Select(E => E.Message)));
+					LogSummary.LogEntries = LogSummary.LogEntries.Except(IgnoredWarnings).ToArray();
+				}
+				if (IgnoredErrors.Any())
+				{
+					Log.Info("Ignoring {0} errors.", IgnoredErrors.Count());
+					Log.Info("\t{0}", string.Join("\n\t", IgnoredErrors.Select(E => E.Message)));
+					LogSummary.LogEntries = LogSummary.LogEntries.Except(IgnoredErrors).ToArray();
+				}
+			}
+
+
+			return LogSummary;
+		}
+
+		protected override UnrealProcessResult GetExitCodeAndReason(StopReason InReason, UnrealLog InLogSummary, UnrealRoleArtifacts InArtifacts, out string ExitReason, out int ExitCode)
+		{
+			// Check for login failure
+			UnrealLogParser Parser = new UnrealLogParser(InArtifacts.AppInstance.StdOut);
+
+			ExitReason = "";
+			ExitCode = -1;
+
+			foreach (HandledError ErrorToCheck in HandledErrors)
+			{
+				string[] MatchingErrors = Parser.GetErrors(ErrorToCheck.CategoryName).Where(E => E.Contains(ErrorToCheck.ClientErrorString)).ToArray();
+				if (MatchingErrors.Length > 0)
+				{
+					ExitReason = string.Format("Test Error: {0} {1}", ErrorToCheck.GauntletErrorString, ErrorToCheck.Verbose ? "\"" + MatchingErrors[0] + "\"" : "");
+					ExitCode = -1;
+					return UnrealProcessResult.TestFailure;
+				}
+			}
+
+			return base.GetExitCodeAndReason(InReason, InLogSummary, InArtifacts, out ExitReason, out ExitCode);
+		}
+
+		public override ITestReport CreateReport(TestResult Result, UnrealTestContext Context, UnrealBuildSource Build, IEnumerable<UnrealRoleResult> Artifacts, string ArtifactPath)
+		{
+			if (Result == TestResult.Passed)
+			{
+				// Our artifacts from each iteration such as the client log will be overwritten by subsequent iterations so we need to copy them out to a temp dir
+				// to preserve them until we're ready to make our report on the final iteration.
+				CopyPerfFilesToTempDir(ArtifactPath);
+
+				if (GetCurrentPass() < (GetNumPasses() - 1))
+				{
+					Log.Info($"Skipping Csv report generator until final pass. On pass {GetCurrentPass() + 1} of {GetNumPasses()}.");
+					return base.CreateReport(Result, Context, Build, Artifacts, ArtifactPath);
+				}
+
+				// Local report generation is an example of how to use the PerfReportTool.
+				if (!Globals.Params.ParseParam("NoLocalReports"))
+				{
+					// NOTE: This does not currently work with long paths due to the CsvTools not properly supporting them.
+					Log.Info("Generating performance reports using PerfReportTool.");
+					GenerateLocalPerfReport(Context.GetRoleContext(UnrealTargetRole.Client).Platform, ArtifactPath);
+				}
+
+				if (Globals.Params.ParseParam("PerfReportServer") &&
+					!Globals.Params.ParseParam("SkipPerfReportServer"))
+				{
+					Dictionary<string, dynamic> CommonDataSourceFields = new Dictionary<string, dynamic>
+					{
+						{ "HordeJobUrl", Globals.Params.ParseValue("JobDetails", null) }
+					};
+
+					Log.Info("Creating perf server importer with build name {BuildName}", OriginalBuildName);
+					string DataSourceName = string.Format("Automation.{0}.Performance", GetCachedConfiguration().ProjectName);
+					string ImportDirOverride = Globals.Params.ParseValue("PerfReportServerImportDir", null);
+					ICsvImporter Importer = ReportGenUtils.CreatePerfReportServerImporter(DataSourceName, OriginalBuildName, CommandUtils.IsBuildMachine, ImportDirOverride, CommonDataSourceFields);
+					if (Importer != null)
+					{
+						// Recursively grab all the csv files we copied to the temp dir and convert them to binary.
+						List<FileInfo> AllBinaryCsvFiles = ReportGenUtils.CollectAndConvertCsvFilesToBinary(TempPerfCSVDir.FullName);
+						if (AllBinaryCsvFiles.Count == 0)
+						{
+							throw new AutomationException($"No Csv files found in {TempPerfCSVDir}");
+						}
+
+						// The corresponding log for each csv sits in the same subdirectory as the csv file itself.
+						IEnumerable<CsvImportEntry> ImportEntries = AllBinaryCsvFiles
+							.Select(CsvFile => new CsvImportEntry(CsvFile.FullName, Path.Combine(CsvFile.Directory.FullName, "ClientOutput.log")));
+
+						// Create the import batch
+						Importer.Import(ImportEntries);
+					}
+
+					// Cleanup the temp dir
+					TempPerfCSVDir.Delete(recursive: true);
+				}
+			}
+			else
+			{
+				Logger.LogWarning("Skipping performance report generation because the perf report test failed.");
+			}
+
+			return base.CreateReport(Result, Context, Build, Artifacts, ArtifactPath);
+		}
+
+		/// <summary>
+		/// Produces a detailed csv report using PerfReportTool.
+		/// Also, stores perf data in the perf cache, and generates a historic report using the data the cache contains.
+		/// </summary>
+		private void GenerateLocalPerfReport(UnrealTargetPlatform Platform, string ArtifactPath)
+		{
+			var ReportCacheDir = GetCachedConfiguration().PerfCacheRoot;	// see if this is appropriate
+
+			var ToolPath = FileReference.Combine(Unreal.EngineDirectory, "Binaries", "DotNET", "CsvTools", "PerfreportTool.exe");
+			if (!FileReference.Exists(ToolPath))
+			{
+				Logger.LogError("Failed to find perf report utility at this path: \"{ToolPath}\".", ToolPath);
+				return;
+			}
+
+			var ReportConfigDir = Path.Combine(Unreal.RootDirectory.FullName, "Samples", "Showcases", GetCachedConfiguration().ProjectName, "Build", "Scripts", "PerfReport");
+			var ReportPath = Path.Combine(ArtifactPath, "Reports", "Performance");
+
+			// Csv files may have been output in one of two places.
+			// Check both...
+			var CsvsPaths = new[]
+			{
+				Path.Combine(ArtifactPath, "Client", "Profiling", "FPSChartStats"),
+				Path.Combine(ArtifactPath, "Client", "Settings", GetCachedConfiguration().ProjectName, "Saved", "Profiling", "FPSChartStats")
+			};
+
+			var DiscoveredCsvs = new List<string>();
+			foreach (var CsvsPath in CsvsPaths)
+			{
+				if (Directory.Exists(CsvsPath))
+				{
+					DiscoveredCsvs.AddRange(
+						from CsvFile in Directory.GetFiles(CsvsPath, "*.csv", SearchOption.AllDirectories)
+						where CsvFile.Contains("csvprofile", StringComparison.InvariantCultureIgnoreCase)
+						select CsvFile);
+				}
+			}
+
+			if (DiscoveredCsvs.Count == 0)
+			{
+				Logger.LogError("Test completed successfully but no csv profiling results were found. Searched paths were:\r\n  {Paths}", string.Join("\r\n  ", CsvsPaths.Select(s => $"\"{s}\"")));
+				return;
+			}
+
+			// Find the newest csv file and get its directory
+			// (PerfReportTool will only output cached data in -csvdir mode)
+			var NewestFile =
+				(from CsvFile in DiscoveredCsvs
+				 let Timestamp = File.GetCreationTimeUtc(CsvFile)
+				 orderby Timestamp descending
+				 select CsvFile).First();
+			var NewestDir = Path.GetDirectoryName(NewestFile);
+
+			Log.Info("Using perf report cache directory \"{ReportCacheDir}\".", ReportCacheDir);
+			Log.Info("Using perf report output directory \"{ReportPath}\".", ReportPath);
+			Log.Info("Using csv results directory \"{NewestDir}\". Generating historic perf report data...", NewestDir);
+
+			// Make sure the cache and output directories exist
+			if (!Directory.Exists(ReportCacheDir))
+			{
+				try { Directory.CreateDirectory(ReportCacheDir); }
+				catch (Exception Ex)
+				{
+					Logger.LogError("Failed to create perf report cache directory \"{ReportCacheDir}\". {Ex}", ReportCacheDir, Ex);
+					return;
+				}
+			}
+			if (!Directory.Exists(ReportPath))
+			{
+				try { Directory.CreateDirectory(ReportPath); }
+				catch (Exception Ex)
+				{
+					Logger.LogError("Failed to create perf report output directory \"{ReportPath}\". {Ex}", ReportPath, Ex);
+					return;
+				}
+			}
+
+			// Win64 is actually called "Windows" in csv profiles
+			var PlatformNameFilter = Platform == UnrealTargetPlatform.Win64 ? "Windows" : $"{Platform}";
+
+			// Produce the detailed report, and update the perf cache
+			CommandUtils.RunAndLog(ToolPath.FullName, $"-csvdir \"{NewestDir}\" -o \"{ReportPath}\" -reportxmlbasedir \"{ReportConfigDir}\" -summaryTableCache \"{ReportCacheDir}\" -searchpattern csvprofile* -metadatafilter platform=\"{PlatformNameFilter}\"", out int ErrorCode);
+			if (ErrorCode != 0)
+			{
+				Logger.LogError("PerfReportTool returned error code \"{ErrorCode}\" while generating detailed report.", ErrorCode);
+			}
+
+			// Now generate the all-time historic summary report
+			HistoricReport("HistoricReport_AllTime", new[]
+			{
+				$"platform={PlatformNameFilter}"
+			});
+
+			// 14 days historic report
+			HistoricReport($"HistoricReport_14Days", new[]
+			{
+				$"platform={PlatformNameFilter}",
+				$"starttimestamp>={DateTimeOffset.Now.ToUnixTimeSeconds() - (14 * 60L * 60L * 24L)}"
+			});
+
+			// 7 days historic report
+			HistoricReport($"HistoricReport_7Days", new[]
+			{
+				$"platform={PlatformNameFilter}",
+				$"starttimestamp>={DateTimeOffset.Now.ToUnixTimeSeconds() - (7 * 60L * 60L * 24L)}"
+			});
+
+			void HistoricReport(string Name, IEnumerable<string> Filter)
+			{
+				var Args = new[]
+				{
+					$"-summarytablecachein \"{ReportCacheDir}\"",
+					$"-summaryTableFilename \"{Name}.html\"",
+					$"-reportxmlbasedir \"{ReportConfigDir}\"",
+					$"-o \"{ReportPath}\"",
+					$"-metadatafilter \"{string.Join(" and ", Filter)}\"",
+					"-summaryTable autoPerfReportStandard",
+					"-condensedSummaryTable autoPerfReportStandard",
+					"-emailtable",
+					"-recurse"
+				};
+
+				var ArgStr = string.Join(" ", Args);
+
+				CommandUtils.RunAndLog(ToolPath.FullName, ArgStr, out ErrorCode);
+				if (ErrorCode != 0)
+				{
+					Logger.LogError("PerfReportTool returned error code \"{ErrorCode}\" while generating historic report.", ErrorCode);
+				}
+			}
+		}
+
+		private void CopyPerfFilesToTempDir(string ArtifactPath)
+		{
+			if (!TempPerfCSVDir.Exists)
+			{
+				Log.Info("Creating temp perf csv dir: {TempPerfCSVDir}", TempPerfCSVDir);
+				TempPerfCSVDir.Create();
+			}
+
+			string ClientArtifactDir = Path.Combine(ArtifactPath, "Client");
+			string ClientLogPath = Path.Combine(ClientArtifactDir, "ClientOutput.log");
+
+			// The FPSChartStats folder can vary in location depending on the platform, so use the helper to find the best match.
+			string FPSChartsPath = PathUtils.FindRelevantPath(ClientArtifactDir, "Profiling", "FPSChartStats");
+			if (string.IsNullOrEmpty(FPSChartsPath))
+			{
+				Logger.LogWarning("Failed to find FPSCharts folder in {ClientArtifactDir}", ClientArtifactDir);
+				return;
+			}
+
+			// CsvTools use .NET Framework 4.8 currently so we must explicitly create long paths for them to work.
+			FPSChartsPath = Gauntlet.Utils.SystemHelpers.GetFullyQualifiedPath(FPSChartsPath);
+
+			// Grab all the csv files that have valid metadata.
+			// We don't want to convert to binary in place as the legacy reports require the raw csv.
+			List<FileInfo> CsvFiles = ReportGenUtils.CollectValidCsvFiles(FPSChartsPath);
+			if (CsvFiles.Count > 0)
+			{
+				// We only want to copy the latest file as the other will have already been copied when this was run for those iterations.
+				CsvFiles.SortBy(Info => Info.LastWriteTimeUtc);
+				FileInfo LatestCsvFile = CsvFiles.Last();
+
+				// Create a subdir for each pass as we want to store the csv and log together in the same dir to make it easier to find them later.
+				string PassDir = Path.Combine(TempPerfCSVDir.FullName, $"PerfCsv_Pass_{GetCurrentPass()}");
+				Directory.CreateDirectory(PassDir);
+
+				FileInfo LogFileInfo = new FileInfo(ClientLogPath);
+				if (LogFileInfo.Exists)
+				{
+					string LogDestPath = Path.Combine(PassDir, LogFileInfo.Name);
+					Log.Info("Copying Log {ClientLogPath} To {LogDest}", ClientLogPath, LogDestPath);
+					LogFileInfo.CopyTo(LogDestPath);
+				}
+				else
+				{
+					Logger.LogWarning("No log file was found at {ClientLogPath}", ClientLogPath);
+				}
+
+				string CsvDestPath = Path.Combine(PassDir, LatestCsvFile.Name);
+				Log.Info("Copying Csv {CsvPath} To {CsvDestPath}", LatestCsvFile.FullName, CsvDestPath);
+				LatestCsvFile.CopyTo(CsvDestPath);
+			}
+			else
+			{
+				Logger.LogWarning("No valid csv files found in {FPSChartsPath}", FPSChartsPath);
+			}
+		}
+
+		/// <summary>
+		/// Returns the cached version of our config. Avoids repeatedly calling GetConfiguration() on derived nodes
+		/// </summary>
+		/// <returns></returns>
+		private TConfigClass GetCachedConfiguration()
+		{
+			if (CachedConfig == null)
+			{
+				return GetConfiguration();
+			}
+
+			return CachedConfig;
+		}
+	}
+}
