AutoTestPerfConfig: hooking up CitySample pass 2.

- Moving more of the sample-specific functionality into the plugin

diff --git a/Engine/Plugins/Performance/AutomatedPerfTesting/Build/Scripts/AutomatedPerfTestConfig.cs b/Engine/Plugins/Performance/AutomatedPerfTesting/Build/Scripts/AutomatedPerfTestConfig.cs
--- a/Engine/Plugins/Performance/AutomatedPerfTesting/Build/Scripts/AutomatedPerfTestConfig.cs
+++ b/Engine/Plugins/Performance/AutomatedPerfTesting/Build/Scripts/AutomatedPerfTestConfig.cs
@@ -17,12 +17,18 @@ namespace AutomatedPerfTest
 		[AutoParamWithNames("AutomatedPerfTest.ProjectName")]
 		public string ProjectName = "";
 
+		/// <summary>
+		/// Name of the project (TODO: parse from the build metadata? What about github builds?)
+		/// </summary>
+		[AutoParamWithNames("AutomatedPerfTest.DataSourceName")]
+		public string DataSourceName = "";
+
 		/// <summary>
 		/// If we're running on the build machine
 		/// </summary>
 		[AutoParamWithNames(false, "AutomatedPerfTest.IsBuildMachine")]
 		public bool IsBuildMachine;
-		
+
 		/// <summary>
 		/// Name of the test, useful for identifying it later
 		/// </summary>
diff --git a/Engine/Plugins/Performance/AutomatedPerfTesting/Build/Scripts/AutomatedPerfTestNode.cs b/Engine/Plugins/Performance/AutomatedPerfTesting/Build/Scripts/AutomatedPerfTestNode.cs
--- a/Engine/Plugins/Performance/AutomatedPerfTesting/Build/Scripts/AutomatedPerfTestNode.cs
+++ b/Engine/Plugins/Performance/AutomatedPerfTesting/Build/Scripts/AutomatedPerfTestNode.cs
@@ -26,8 +26,13 @@ namespace AutomatedPerfTest
 	{
 		public AutomatedPerfTestNode(UnrealTestContext InContext) : base(InContext)
 		{
+			// We need to save off the build name as if this is a preflight that suffix will be stripped
+			// after GetConfiguration is called. This will cause a mismatch in CreateReport.
+			OriginalBuildName = Globals.Params.ParseValue("BuildName", InContext.BuildInfo.BuildName);
+			Log.Info("Setting OriginalBuildName to {OriginalBuildName}", OriginalBuildName);
+
 			TestGuid = Guid.NewGuid();
-			Gauntlet.Log.Info("Your Test GUID is :\n" + TestGuid.ToString() + '\n');
+			Log.Info("Your Test GUID is :\n" + TestGuid.ToString() + '\n');
 
 			InitHandledErrors();
 
@@ -223,55 +228,88 @@ namespace AutomatedPerfTest
 		{
 			if (Result == TestResult.Passed)
 			{
-				// Our artifacts from each iteration such as the client log will be overwritten by subsequent iterations so we need to copy them out to a temp dir
-				// to preserve them until we're ready to make our report on the final iteration.
-				CopyPerfFilesToTempDir(ArtifactPath);
-
-				if (GetCurrentPass() < (GetNumPasses() - 1))
-				{
-					Log.Info($"Skipping Csv report generator until final pass. On pass {GetCurrentPass() + 1} of {GetNumPasses()}.");
-					return base.CreateReport(Result, Context, Build, Artifacts, ArtifactPath);
-				}
-
-				// Local report generation is an example of how to use the PerfReportTool.
-				if (!Globals.Params.ParseParam("NoLocalReports"))
+				if (GetConfiguration().DoCSVProfiler)
 				{
-					// NOTE: This does not currently work with long paths due to the CsvTools not properly supporting them.
-					Log.Info("Generating performance reports using PerfReportTool.");
-					GenerateLocalPerfReport(Context.GetRoleContext(UnrealTargetRole.Client).Platform, ArtifactPath);
-				}
+					// Our artifacts from each iteration such as the client log will be overwritten by subsequent iterations so we need to copy them out to a temp dir
+					// to preserve them until we're ready to make our report on the final iteration.
+					CopyPerfFilesToTempDir(ArtifactPath);
 
-				if (Globals.Params.ParseParam("PerfReportServer") &&
-					!Globals.Params.ParseParam("SkipPerfReportServer"))
-				{
-					Dictionary<string, dynamic> CommonDataSourceFields = new Dictionary<string, dynamic>
+					// Local report generation is an example of how to use the PerfReportTool.
+					if (!Globals.Params.ParseParam("NoLocalReports"))
 					{
-						{ "HordeJobUrl", Globals.Params.ParseValue("JobDetails", null) }
-					};
-
-					Log.Info("Creating perf server importer with build name {BuildName}", OriginalBuildName);
-					string DataSourceName = string.Format("Automation.{0}.Performance", GetCachedConfiguration().ProjectName);
-					string ImportDirOverride = Globals.Params.ParseValue("PerfReportServerImportDir", null);
-					ICsvImporter Importer = ReportGenUtils.CreatePerfReportServerImporter(DataSourceName, OriginalBuildName, CommandUtils.IsBuildMachine, ImportDirOverride, CommonDataSourceFields);
-					if (Importer != null)
+						// NOTE: This does not currently work with long paths due to the CsvTools not properly supporting them.
+						Log.Info("Generating performance reports using PerfReportTool.");
+						GenerateLocalPerfReport(Context.GetRoleContext(UnrealTargetRole.Client).Platform, ArtifactPath);
+					}
+
+					if (Globals.Params.ParseParam("PerfReportServer") &&
+						!Globals.Params.ParseParam("SkipPerfReportServer"))
 					{
-						// Recursively grab all the csv files we copied to the temp dir and convert them to binary.
-						List<FileInfo> AllBinaryCsvFiles = ReportGenUtils.CollectAndConvertCsvFilesToBinary(TempPerfCSVDir.FullName);
-						if (AllBinaryCsvFiles.Count == 0)
+						Dictionary<string, dynamic> CommonDataSourceFields = new Dictionary<string, dynamic>
+						{
+							{"HordeJobUrl", Globals.Params.ParseValue("JobDetails", null)}
+						};
+
+						Log.Info("Creating perf server importer with build name {BuildName}", OriginalBuildName);
+						string DataSourceName = GetConfiguration().DataSourceName;
+						string ImportDirOverride = Globals.Params.ParseValue("PerfReportServerImportDir", null);
+						ICsvImporter Importer = ReportGenUtils.CreatePerfReportServerImporter(DataSourceName, OriginalBuildName,
+							CommandUtils.IsBuildMachine, ImportDirOverride, CommonDataSourceFields);
+						if (Importer != null)
 						{
-							throw new AutomationException($"No Csv files found in {TempPerfCSVDir}");
+							// Recursively grab all the csv files we copied to the temp dir and convert them to binary.
+							List<FileInfo> AllBinaryCsvFiles = ReportGenUtils.CollectAndConvertCsvFilesToBinary(TempPerfCSVDir.FullName);
+							if (AllBinaryCsvFiles.Count == 0)
+							{
+								throw new AutomationException($"No Csv files found in {TempPerfCSVDir}");
+							}
+
+							// The corresponding log for each csv sits in the same subdirectory as the csv file itself.
+							IEnumerable<CsvImportEntry> ImportEntries = AllBinaryCsvFiles
+								.Select(CsvFile => new CsvImportEntry(CsvFile.FullName, Path.Combine(CsvFile.Directory.FullName, "ClientOutput.log")));
+
+							// todo update this so it associates videos with the correct CSVs
+							IEnumerable<CsvImportEntry> CsvImportEntries = ImportEntries as CsvImportEntry[] ?? ImportEntries.ToArray();
+							if (GetConfiguration().DoInsightsTrace)
+							{
+								string InsightsFilename = Path.GetFileNameWithoutExtension(CsvImportEntries.First().CsvFilename)
+									.Replace(".csv", ".utrace");
+								string InsightsFilePath = Path.Combine(ArtifactPath, "Profiling", InsightsFilename);
+								if (File.Exists(InsightsFilePath))
+								{
+									CsvImportEntries.First().AddAdditionalFile("Insights", InsightsFilePath);
+								}
+								else
+								{
+									Log.Warning("Insights was requested, but no match insights trace file was found {InsightsFilename}",
+										InsightsFilePath);
+								}
+							}
+
+							if (GetConfiguration().DoVideoCapture)
+							{
+								string VideoPath = Path.Combine(ArtifactPath, "Client", "Videos");
+								string[] VideoFiles = Directory.GetFiles(VideoPath, "*.mp4");
+								if (VideoFiles.Length > 0)
+								{
+									foreach (var VideoFile in VideoFiles)
+									{
+										CsvImportEntries.First().AddAdditionalFile("Video", Path.Combine(VideoPath, VideoFile));
+									}
+								}
+								else
+								{
+									Log.Warning("Video capture was requested, but no videos were found in path {VideoPath}", VideoPath);
+								}
+							}
+
+							// Create the import batch
+							Importer.Import(CsvImportEntries);
 						}
 
-						// The corresponding log for each csv sits in the same subdirectory as the csv file itself.
-						IEnumerable<CsvImportEntry> ImportEntries = AllBinaryCsvFiles
-							.Select(CsvFile => new CsvImportEntry(CsvFile.FullName, Path.Combine(CsvFile.Directory.FullName, "ClientOutput.log")));
-
-						// Create the import batch
-						Importer.Import(ImportEntries);
+						// Cleanup the temp dir
+						TempPerfCSVDir.Delete(recursive: true);
 					}
-
-					// Cleanup the temp dir
-					TempPerfCSVDir.Delete(recursive: true);
 				}
 			}
 			else
@@ -304,8 +342,7 @@ namespace AutomatedPerfTest
 			// Check both...
 			var CsvsPaths = new[]
 			{
-				Path.Combine(ArtifactPath, "Client", "Profiling", "FPSChartStats"),
-				Path.Combine(ArtifactPath, "Client", "Settings", GetCachedConfiguration().ProjectName, "Saved", "Profiling", "FPSChartStats")
+				Path.Combine(ArtifactPath, "Client", "Profiling", "CSV")
 			};
 
 			var DiscoveredCsvs = new List<string>();
@@ -315,7 +352,6 @@ namespace AutomatedPerfTest
 				{
 					DiscoveredCsvs.AddRange(
 						from CsvFile in Directory.GetFiles(CsvsPath, "*.csv", SearchOption.AllDirectories)
-						where CsvFile.Contains("csvprofile", StringComparison.InvariantCultureIgnoreCase)
 						select CsvFile);
 				}
 			}
@@ -425,20 +461,17 @@ namespace AutomatedPerfTest
 			string ClientArtifactDir = Path.Combine(ArtifactPath, "Client");
 			string ClientLogPath = Path.Combine(ClientArtifactDir, "ClientOutput.log");
 
-			// The FPSChartStats folder can vary in location depending on the platform, so use the helper to find the best match.
-			string FPSChartsPath = PathUtils.FindRelevantPath(ClientArtifactDir, "Profiling", "FPSChartStats");
-			if (string.IsNullOrEmpty(FPSChartsPath))
+
+			string CSVPath = PathUtils.FindRelevantPath(ClientArtifactDir, "Profiling", "CSV");
+			if (string.IsNullOrEmpty(CSVPath))
 			{
-				Logger.LogWarning("Failed to find FPSCharts folder in {ClientArtifactDir}", ClientArtifactDir);
+				Log.Warning("Failed to find CSV folder folder in {ClientArtifactDir}", ClientArtifactDir);
 				return;
 			}
 
-			// CsvTools use .NET Framework 4.8 currently so we must explicitly create long paths for them to work.
-			FPSChartsPath = Gauntlet.Utils.SystemHelpers.GetFullyQualifiedPath(FPSChartsPath);
-
 			// Grab all the csv files that have valid metadata.
 			// We don't want to convert to binary in place as the legacy reports require the raw csv.
-			List<FileInfo> CsvFiles = ReportGenUtils.CollectValidCsvFiles(FPSChartsPath);
+			List<FileInfo> CsvFiles = ReportGenUtils.CollectValidCsvFiles(CSVPath);
 			if (CsvFiles.Count > 0)
 			{
 				// We only want to copy the latest file as the other will have already been copied when this was run for those iterations.
@@ -454,21 +487,80 @@ namespace AutomatedPerfTest
 				{
 					string LogDestPath = Path.Combine(PassDir, LogFileInfo.Name);
 					Log.Info("Copying Log {ClientLogPath} To {LogDest}", ClientLogPath, LogDestPath);
-					LogFileInfo.CopyTo(LogDestPath);
+					LogFileInfo.CopyTo(LogDestPath, true);
 				}
 				else
 				{
-					Logger.LogWarning("No log file was found at {ClientLogPath}", ClientLogPath);
+					Log.Warning("No log file was found at {ClientLogPath}", ClientLogPath);
 				}
 
 				string CsvDestPath = Path.Combine(PassDir, LatestCsvFile.Name);
 				Log.Info("Copying Csv {CsvPath} To {CsvDestPath}", LatestCsvFile.FullName, CsvDestPath);
-				LatestCsvFile.CopyTo(CsvDestPath);
+				LatestCsvFile.CopyTo(CsvDestPath, true);
 			}
 			else
 			{
-				Logger.LogWarning("No valid csv files found in {FPSChartsPath}", FPSChartsPath);
+				Log.Warning("No valid csv files found in {CSVPath}", CSVPath);
+			}
+		}
+
+		protected virtual string GetSubtestName()
+		{
+			return "Performance";
+		}
+
+		public override TConfigClass GetConfiguration()
+		{
+			TConfigClass Config = base.GetConfiguration();
+			Config.MaxDuration = Context.TestParams.ParseValue("MaxDuration", 60 * 60);  // 1 hour max
+
+			UnrealTestRole ClientRole = Config.RequireRole(UnrealTargetRole.Client);
+			// the controller will be added by the subclasses
+
+			ClientRole.CommandLineParams.Add("-deterministic");
+
+			Log.Info("AutomatedPerfTestNode<>.GetConfiguration(): Config.DoFPSChart={0}, Config.DoCSVProfiler={1}, Config.DoVideoCapture={2}", Config.DoFPSChart, Config.DoCSVProfiler, Config.DoVideoCapture);
+
+			if (string.IsNullOrEmpty(Config.DataSourceName))
+			{
+				Config.DataSourceName = string.Format("Automation.{0}.Performance", Config.ProjectName);
+			}
+
+			if (Config.DoFPSChart)
+			{
+				ClientRole.CommandLineParams.Add("AutomatedPerfTest.DoFPSChart");
 			}
+
+			if (Config.DoCSVProfiler)
+			{
+				ClientRole.CommandLineParams.Add("AutomatedPerfTest.DoCSVProfiler");
+				ClientRole.CommandLineParams.Add("csvGpuStats");
+
+				// Add CSV metadata
+				List<string> CsvMetadata = new List<string>
+				{
+					string.Format("testname={0}", Config.ProjectName),
+					"gauntletTestType=AutomatedPerfTest",
+					string.Format("gauntletSubTest={0}", GetSubtestName()),
+					"testBuildIsPreflight=" + (ReportGenUtils.IsTestingPreflightBuild(OriginalBuildName) ? "1" : "0"),
+					"testBuildVersion=" + OriginalBuildName
+				};
+
+				if (!string.IsNullOrEmpty(Context.BuildInfo.Branch) && Context.BuildInfo.Changelist != 0)
+				{
+					CsvMetadata.Add("branch=" + Context.BuildInfo.Branch);
+					CsvMetadata.Add("changelist=" + Context.BuildInfo.Changelist);
+				}
+
+				ClientRole.CommandLineParams.Add("csvMetadata", "\"" + String.Join(",", CsvMetadata) + "\"");
+			}
+
+			if (Config.DoVideoCapture)
+			{
+				ClientRole.CommandLineParams.Add("AutomatedPerfTest.DoVideoCapture");
+			}
+
+			return Config;
 		}
 
 		/// <summary>
@@ -485,4 +577,26 @@ namespace AutomatedPerfTest
 			return CachedConfig;
 		}
 	}
+
+	/// <summary>
+	/// Implementation of a Gauntlet TestNode for AutomatedPerfTest plugin
+	/// </summary>
+	/// <typeparam name="TConfigClass"></typeparam>
+	public abstract class AutomatedSequencePerfTestNode<TConfigClass> : AutomatedPerfTestNode<TConfigClass>
+		where TConfigClass : AutomatedPerfTestConfigBase, new()
+	{
+		public AutomatedSequencePerfTestNode(UnrealTestContext InContext) : base(InContext)
+		{
+		}
+
+		public override TConfigClass GetConfiguration()
+		{
+			TConfigClass Config = base.GetConfiguration();
+
+			UnrealTestRole ClientRole = Config.RequireRole(UnrealTargetRole.Client);	// should get an existing role
+			ClientRole.Controllers.Add("AutomatedSequencePerfTest");
+
+			return Config;
+		}
+	}
 }
