Change variable name for consistency

#rnx

diff --git a/Engine/Plugins/Performance/AutomatedPerfTesting/Build/Scripts/AutomatedPerfTestNode.cs b/Engine/Plugins/Performance/AutomatedPerfTesting/Build/Scripts/AutomatedPerfTestNode.cs
--- a/Engine/Plugins/Performance/AutomatedPerfTesting/Build/Scripts/AutomatedPerfTestNode.cs
+++ b/Engine/Plugins/Performance/AutomatedPerfTesting/Build/Scripts/AutomatedPerfTestNode.cs
@@ -36,12 +36,12 @@ namespace AutomatedPerfTest
 
 			InitHandledErrors();
 
-			CitySampleLogParser = null;
+			LogParser = null;
 		}
 
 		public override bool StartTest(int Pass, int InNumPasses)
 		{
-			CitySampleLogParser = null;
+			LogParser = null;
 			return base.StartTest(Pass, InNumPasses);
 		}
 
@@ -80,7 +80,7 @@ namespace AutomatedPerfTest
 		/// <summary>
 		/// Track client log messages that have been written to the test logs.
 		/// </summary>
-		private UnrealLogStreamParser CitySampleLogParser;
+		private UnrealLogStreamParser LogParser;
 
 		/// <summary>
 		// Temporary directory for perf report CSVs
@@ -121,13 +121,13 @@ namespace AutomatedPerfTest
 
 			if (App != null)
 			{
-				if (CitySampleLogParser == null)
+				if (LogParser == null)
 				{
-					CitySampleLogParser = new UnrealLogStreamParser(App.GetLogBufferReader());
+					LogParser = new UnrealLogStreamParser(App.GetLogBufferReader());
 				}
-				CitySampleLogParser.ReadStream();
+				LogParser.ReadStream();
 				string LogChannelName = Context.BuildInfo.ProjectName + "Test";
-				List<string> TestLines = CitySampleLogParser.GetLogFromChannel(LogChannelName).ToList();
+				List<string> TestLines = LogParser.GetLogFromChannel(LogChannelName).ToList();
 
 				string LogCategory = "Log" + LogChannelName;
 				string LogCategoryError = LogCategory + ": Error:";
