Filter log channels on partial match

#rnx

diff --git a/Engine/Plugins/Performance/AutomatedPerfTesting/Build/Scripts/AutomatedPerfTestNode.cs b/Engine/Plugins/Performance/AutomatedPerfTesting/Build/Scripts/AutomatedPerfTestNode.cs
--- a/Engine/Plugins/Performance/AutomatedPerfTesting/Build/Scripts/AutomatedPerfTestNode.cs
+++ b/Engine/Plugins/Performance/AutomatedPerfTesting/Build/Scripts/AutomatedPerfTestNode.cs
@@ -127,7 +127,7 @@ namespace AutomatedPerfTest
 				}
 				LogParser.ReadStream();
 				string LogChannelName = Context.BuildInfo.ProjectName + "Test";
-				List<string> TestLines = LogParser.GetLogFromChannel(LogChannelName).ToList();
+				List<string> TestLines = LogParser.GetLogFromChannel(LogChannelName, false).ToList();
 
 				string LogCategory = "Log" + LogChannelName;
 				string LogCategoryError = LogCategory + ": Error:";
