AutomatedPerfTest plugin - default to collecting the CSV profile.

diff --git a/Engine/Plugins/Performance/AutomatedPerfTesting/Build/Scripts/AutomatedPerfTestConfig.cs b/Engine/Plugins/Performance/AutomatedPerfTesting/Build/Scripts/AutomatedPerfTestConfig.cs
--- a/Engine/Plugins/Performance/AutomatedPerfTesting/Build/Scripts/AutomatedPerfTestConfig.cs
+++ b/Engine/Plugins/Performance/AutomatedPerfTesting/Build/Scripts/AutomatedPerfTestConfig.cs
@@ -54,7 +54,7 @@ namespace AutomatedPerfTest
 		/// <summary>
 		/// If we're running on the build machine
 		/// </summary>
-		[AutoParamWithNames(false, "AutomatedPerfTest.DoCSVProfiler")]
+		[AutoParamWithNames(true, "AutomatedPerfTest.DoCSVProfiler")]
 		public bool DoCSVProfiler;
 		
 		/// <summary>
