Add AutomatedSequenceTest with Perf/LLM/Insights subtypes to CitySample.

- Simplified and reduced the amount of boilerplate needed, particularly C# code with Gauntlet nodes is out, and some BuildGraph config files.
- Also add Windows_Vulkan platform for CitySample Vulkan test.

diff --git a/Engine/Build/Graph/Tasks/BuildAndTestProject.xml b/Engine/Build/Graph/Tasks/BuildAndTestProject.xml
--- a/Engine/Build/Graph/Tasks/BuildAndTestProject.xml
+++ b/Engine/Build/Graph/Tasks/BuildAndTestProject.xml
@@ -219,7 +219,27 @@
 			<Regex Pattern="(.*):(.*)" Capture="Left;Right" Input="$(TestNameTitle)"/>
 			<Property Name="TestNameTitle" Value="$(Left)_$(Right)"/>
 		</Do>
+		<!-- We want to use a user-provided name if any. Expected format of the node name is 'FooTest(arg1+arg2+arg3)[Test Name]' and we want to use Test Name if provided -->
+		<!-- This is intentionally after the previous test to override it. -->
+		<Do If="Contains('$(TestNameTitle)', '[')">
+			<Property Name="HumanReadableName" Value=""/>
+			<Regex Pattern="\[(.*)\]" Capture="HumanReadableName" Input="$(TestNameTitle)"/>
+			<Property Name="TestNameTitle" Value="$(HumanReadableName)"/>
+		</Do>
 	</Macro>
+	
+	<!-- Split TestName from Foo.Test{Arg1 Arg2 Arg3} into Foo.Test and Arg1 Arg2 Arg3 -->
+	<Macro Name="SplitCombinedTestNameAndArgs">
+		<!-- Expect properties named "TestNameBeforeParentheses" and "TestArgumentsWithinParentheses" to be defined beforehand -->
+		<Do If="Contains('$(TestNameBeforeParentheses)', '{')">
+			<Property Name="Command" Value=""/>
+			<Property Name="Args" Value=""/>
+			<Regex Pattern="(.*)\{(.*)\}" Capture="Command;Args" Input="$(TestNameBeforeParentheses)"/>
+			<Property Name="TestNameBeforeParentheses" Value="$(Command)"/>
+			<Property Name="TestArgumentsWithinParentheses" Value="$(Args)"/>
+		</Do>
+	</Macro>
+	
 
 	<!-- These properties will be filled in with nodes that are generated based on options -->
 	<Property Name="GeneratedToolNodes" Value="" />
@@ -334,8 +354,11 @@
 				<Property Name="TestNodeName" Value="$(PreNodeName)Editor $(EditorPlatform) Test=$(TestNameTitle)" />
 				<Property Name="NodeAnnotations" Value=""/>
 				<Property Name="NodeAnnotations" Value="Workflow=$(AnnotationsTarget)" If="ContainsItem('$(AnnotationsPlatforms)', '$(EditorPlatform)Editor', ';') and ContainsItem('$(AnnotationsTests)', '$(TestName)', ';')"/>
+				<Property Name="TestNameBeforeParentheses" Value="$(TestName)"/>
+				<Property Name="TestArgumentsWithinParentheses" Value=""/>
+				<Expand Name="SplitCombinedTestNameAndArgs"/>
 				<Node Name="$(TestNodeName)" Requires="$(TestPrerequisites)" Annotations="$(NodeAnnotations)">
-					<Command Name="RunUnreal" Arguments="-project=$(TargetProject) -build=Editor -test=&quot;$(TestName)&quot; $(GauntletReportArgs) $(GauntletEditorTestArgs)" />
+					<Command Name="RunUnreal" Arguments="-project=$(TargetProject) -build=Editor -test=&quot;$(TestNameBeforeParentheses)&quot; $(GauntletReportArgs) $(GauntletEditorTestArgs) $(TestArgumentsWithinParentheses)" />
 				</Node>
 				<!-- Add each test node to our generated list -->
 				<Property Name="PlatformEditorTestNodes" Value="$(PlatformEditorTestNodes);$(TestNodeName)" />
@@ -440,8 +463,11 @@
 				<Property Name="TestNodeName" Value="$(PreNodeName)CookedEditor $(EditorPlatform) Test=$(TestNameTitle)" />
 				<Property Name="NodeAnnotations" Value=""/>
 				<Property Name="NodeAnnotations" Value="Workflow=$(AnnotationsTarget)" If="ContainsItem('$(AnnotationsPlatforms)', '$(EditorPlatform)Editor', ';') and ContainsItem('$(AnnotationsTests)', '$(TestName)', ';')"/>
+				<Property Name="TestNameBeforeParentheses" Value="$(TestName)"/>
+				<Property Name="TestArgumentsWithinParentheses" Value=""/>
+				<Expand Name="SplitCombinedTestNameAndArgs"/>
 				<Node Name="$(TestNodeName)" Requires="$(TestPrerequisites)" Annotations="$(NodeAnnotations)">
-					<Command Name="RunUnreal" Arguments="-project=$(TargetProject) -build=&quot;$(NetworkOutputDirectory)/$(EditorPlatformPath)CookedEditor&quot; -test=&quot;$(TestName)&quot; $(GauntletReportArgs) $(CustomModuleRoleArg) $(GauntletCookedEditorTestArgs)" />
+					<Command Name="RunUnreal" Arguments="-project=$(TargetProject) -build=&quot;$(NetworkOutputDirectory)/$(EditorPlatformPath)CookedEditor&quot; -test=&quot;$(TestNameBeforeParentheses)&quot; $(GauntletReportArgs) $(CustomModuleRoleArg) $(GauntletCookedEditorTestArgs) $(TestArgumentsWithinParentheses)" />
 				</Node>
 				<!-- Add each test node to our generated list -->
 				<Property Name="PlatformCookedEditorTestNodes" Value="$(PlatformCookedEditorTestNodes);$(TestNodeName)" />
@@ -793,8 +819,11 @@
 					<Property Name="NodeAnnotations" Value=""/>
 					<Property Name="NodeAnnotations" Value="Workflow=$(AnnotationsTarget)" If="ContainsItem('$(AnnotationsPlatforms)', '$(TargetPlatform)', ';') and ContainsItem('$(AnnotationsTests)', '$(TestName)', ';')"/>
 					<!-- Run the test. No 'Requires' as we run on the same machine locally and use the publish folder otherwise -->
+					<Property Name="TestNameBeforeParentheses" Value="$(TestName)"/>
+					<Property Name="TestArgumentsWithinParentheses" Value=""/>
+					<Expand Name="SplitCombinedTestNameAndArgs"/>
 					<Node Name="$(TestNodeName)" Requires="$(TestPrerequisites)" Annotations="$(NodeAnnotations)">
-						<Command Name="RunUnreal" Arguments="-project=$(TargetProject) -test=&quot;$(TestName)&quot; -platform=$(TargetPlatform) -configuration=$(TestConfiguration) -target=$(TargetName) $(GauntletBuildArg) $(GauntletTestArgs) $(GauntletReportArgs)"/>
+						<Command Name="RunUnreal" Arguments="-project=$(TargetProject) -test=&quot;$(TestNameBeforeParentheses)&quot; -platform=$(TargetPlatform) -configuration=$(TestConfiguration) -target=$(TargetName) $(GauntletBuildArg) $(GauntletTestArgs) $(TestArgumentsWithinParentheses) $(GauntletReportArgs)"/>
 					</Node>
 					<!-- Add it to our list -->
 					<Do If="!ContainsItem('$(TargetPlatformsToSkipTests)', '$(TargetPlatform)', '+')">
diff --git a/Engine/Plugins/Performance/AutomatedPerfTesting/Build/Inc/AutomatedPerfTestCommonSettings.xml b/Engine/Plugins/Performance/AutomatedPerfTesting/Build/Inc/AutomatedPerfTestCommonSettings.xml
--- a/Engine/Plugins/Performance/AutomatedPerfTesting/Build/Inc/AutomatedPerfTestCommonSettings.xml
+++ b/Engine/Plugins/Performance/AutomatedPerfTesting/Build/Inc/AutomatedPerfTestCommonSettings.xml
@@ -14,7 +14,6 @@
 	<!-- Other common arguments --> 
 	<Option Name="TestID" DefaultValue="" Description="A string to use to identify these performance tests" />
 	<Option Name="DeviceProfileOverride" DefaultValue="" Description="Which device profile to use for this test instead of the default value. Will be prepended by PlatformName." />
-	<Option Name="DoInsightsTrace" DefaultValue="false" Restrict="true|false" Description="Whether or not to use Unreal Insights to trace and parse performance results." />
 	<Option Name="DoCSVProfiler" DefaultValue="true" Restrict="true|false" Description="Whether or not to use CSVProfiler to gather and parse performance telemetry." />
 	<Option Name="PerfReportServer" DefaultValue="false" Restrict="true|false" Description="Whether or not to use CSVProfiler to gather and parse performance telemetry." />
 	<Option Name="DoFPSChart" DefaultValue="false" Restrict="true|false" Description="Whether or not to use FPS Chart to gather and parse performance telemetry." />
@@ -23,19 +22,45 @@
 	<Option Name="DataSourceName" DefaultValue="Automation.$(ProjectName).Performance" Description="To which PRS data source should this be uploaded"/>
 	<Option Name="PerfCacheRoot" DefaultValue="$(RootDir)/$(ProjectPath)/PerfCache" Description="Path to the perf cache where all the output data is stored. Defaults to a path in the local project's directory. Override to a common network-accessible path when running on build machines"/>
 	<Option Name="GenerateLocalReports" DefaultValue="false" Restrict="true|false" Description="Whether or not to create a local PerfReportTool report" />
-	
+
+	<!-- Plugin-provided tests -->
+	<Option Name="WithAPTSequenceTests" DefaultValue="true" Restrict="true|false" Description="Whether to run Auto Perf Test sequence flythrough test." />
+	<Option Name="WithAPTMaterialTest" DefaultValue="false" Restrict="true|false" Description="Whether to run Auto Perf Test material test."  />
+
+	<!-- Test sub-types, act as flavors to the above tests -->
+	<Option Name="DoPerf" DefaultValue="false" Restrict="true|false" Description="Whether to run Perf (general performance) test." />
+	<Option Name="DoLLM" DefaultValue="false" Restrict="true|false" Description="Whether to run LLM (memory) test." />
+	<Option Name="DoInsightsTrace" DefaultValue="false" Restrict="true|false" Description="Whether or not to capture Unreal Insights trace during the run." />
+
 	<!-- Construct the common arguments for the performance test and report commands -->
 	<Property Name="AutomatedPerformanceCommonArgs" Value = "-AutomatedPerfTest.UseShippingInsights=$(UseShippingInsights)" />
 	<Property Name="AutomatedPerformanceCommonArgs" Value="$(AutomatedPerformanceCommonArgs) -AutomatedPerfTest.TestName=$(TestID)"/>
 	<Property Name="AutomatedPerformanceCommonArgs" Value="$(AutomatedPerformanceCommonArgs) -AutomatedPerfTest.DeviceProfileOverride=$(DeviceProfileOverride)" If="'$(DeviceProfileOverride)' != ''"/>
-	<Property Name="AutomatedPerformanceCommonArgs" Value="$(AutomatedPerformanceCommonArgs) -AutomatedPerfTest.DoInsightsTrace" If="$(DoInsightsTrace)"/>
 	<Property Name="AutomatedPerformanceCommonArgs" Value="$(AutomatedPerformanceCommonArgs) -AutomatedPerfTest.TraceChannels=$(TraceChannels)" If="$(DoInsightsTrace) and '$(TraceChannels)' != ''"/>
 	<Property Name="AutomatedPerformanceCommonArgs" Value="$(AutomatedPerformanceCommonArgs) -AutomatedPerfTest.DoCSVProfiler" If="$(DoCSVProfiler)"/>
 	<Property Name="AutomatedPerformanceCommonArgs" Value="$(AutomatedPerformanceCommonArgs) -AutomatedPerfTest.DoFPSChart" If="$(DoFPSChart)"/>
+	<!--  disabled until we find a solution for escaping the path names in JSON nodes
 	<Property Name="AutomatedPerformanceCommonArgs" Value="$(AutomatedPerformanceCommonArgs) -AutomatedPerfTest.PerfCacheRoot=$(PerfCacheRoot)"/>
+	-->
 	<Property Name="AutomatedPerformanceCommonArgs" Value="$(AutomatedPerformanceCommonArgs) -PerfReportServer" If="$(PerfReportServer) And $(DoCSVProfiler)"/>
 	<Property Name="AutomatedPerformanceCommonArgs" Value="$(AutomatedPerformanceCommonArgs) -NoLocalReports" If="!$(GenerateLocalReports) And $(DoCSVProfiler)" />
 	<Property Name="AutomatedPerformanceCommonArgs" Value="$(AutomatedPerformanceCommonArgs) -resX=3840 -resY=2160"/>
 	<Property Name="AutomatedPerformanceCommonArgs" Value="$(AutomatedPerformanceCommonArgs) -AutomatedPerfTest.DoVideoCapture" If="$(DoVideoCapture)"/>
 
+	<!-- append plugin-provided tests to GauntletSettings' TargetTestList -->
+	<Do If="$(WithAPTSequenceTests)">
+		<ForEach Name="MapSequenceComboName" Values="$(MapSequenceComboNames)" Separator="+">
+			<Property Name="TargetTestList" Value="$(TargetTestList)AutomatedPerfTest.SequenceTest{-AutomatedPerfTest.SequencePerfTest.MapSequenceName=$(MapSequenceComboName) $(AutomatedPerformanceCommonArgs)}[Sequence Perf $(MapSequenceComboName)]+" If="$(DoPerf)" />
+			<Property Name="TargetTestList" Value="$(TargetTestList)AutomatedPerfTest.SequenceTest{-AutomatedPerfTest.SequencePerfTest.MapSequenceName=$(MapSequenceComboName) -AutomatedPerfTest.DoLLM=true, $(AutomatedPerformanceCommonArgs)}[Sequence LLM $(MapSequenceComboName)]+" If="$(DoLLM)"/>
+			<Property Name="TargetTestList" Value="$(TargetTestList)AutomatedPerfTest.SequenceTest{-AutomatedPerfTest.SequencePerfTest.MapSequenceName=$(MapSequenceComboName) -AutomatedPerfTest.DoInsightsTrace=true, $(AutomatedPerformanceCommonArgs)}[Sequence Insights $(MapSequenceComboName)]+" If="$(DoInsightsTrace)"/>
+		</ForEach>
+	</Do>
+
+	<Do If="$(WithAPTMaterialTest)">
+		<Property Name="TargetTestList" Value="$(TargetTestList)AutomatedPerfTest.MaterialTest($(AutomatedPerformanceCommonArgs))[Material Perf]+" If="$(DoPerf)" />
+		<Property Name="TargetTestList" Value="$(TargetTestList)AutomatedPerfTest.MaterialTest(AutomatedPerfTest.DoLLM=true $(AutomatedPerformanceCommonArgs))[Material LLM]+" If="$(DoLLM)"/>
+		<Property Name="TargetTestList" Value="$(TargetTestList)AutomatedPerfTest.MaterialTest(AutomatedPerfTest.DoInsightsTrace=true $(AutomatedPerformanceCommonArgs))[Material Insights]+" If="$(DoInsightsTrace)"/>
+	</Do>
+
+
 </BuildGraph>
diff --git a/Engine/Plugins/Performance/AutomatedPerfTesting/Build/Inc/AutomatedPerfTestProjectSettings.xml b/Engine/Plugins/Performance/AutomatedPerfTesting/Build/Inc/AutomatedPerfTestProjectSettings.xml
new file mode 100644
--- /dev/null
+++ b/Engine/Plugins/Performance/AutomatedPerfTesting/Build/Inc/AutomatedPerfTestProjectSettings.xml
@@ -0,0 +1,7 @@
+<?xml version="1.0" encoding="utf-8"?>
+<BuildGraph xmlns="http://www.epicgames.com/BuildGraph" xmlns:xsi="http://www.w3.org/2001/XMLSchema-instance" xsi:schemaLocation="http://www.epicgames.com/BuildGraph ../../../../../../../Engine/Build/Graph/Schema.xsd" >
+
+	<!-- These are options for Automated Perf Testing that can - and often should - be overriden in a given project -->
+	<Option Name="MapSequenceComboNames" DefaultValue="" Description="+ separated list of names, referencing the names of Map/Sequence combos defined in Automated Perf Test Project Settings" />
+
+</BuildGraph>
diff --git a/Engine/Plugins/Performance/AutomatedPerfTesting/Build/Scripts/AutomatedPerfTestConfig.cs b/Engine/Plugins/Performance/AutomatedPerfTesting/Build/Scripts/AutomatedPerfTestConfig.cs
--- a/Engine/Plugins/Performance/AutomatedPerfTesting/Build/Scripts/AutomatedPerfTestConfig.cs
+++ b/Engine/Plugins/Performance/AutomatedPerfTesting/Build/Scripts/AutomatedPerfTestConfig.cs
@@ -40,11 +40,11 @@ namespace AutomatedPerfTest
 		public string DeviceProfileOverride = "";
 		
 		/// <summary>
-		/// If we're running on the build machine
+		/// Whether to do an Insights trace
 		/// </summary>
 		[AutoParamWithNames(false, "AutomatedPerfTest.DoInsightsTrace")]
 		public bool DoInsightsTrace;
-		
+
 		/// <summary>
 		/// Which trace channels to test with
 		/// </summary>
@@ -62,7 +62,13 @@ namespace AutomatedPerfTest
 		/// </summary>
 		[AutoParamWithNames(false, "AutomatedPerfTest.DoFPSChart")]
 		public bool DoFPSChart;
-		
+
+		/// <summary>
+		/// Whether to run the test with -llm
+		/// </summary>
+		[AutoParamWithNames(false, "AutomatedPerfTest.DoLLM")]
+		public bool DoLLM;
+
 		/// <summary>
 		/// Let BuildGraph tell us where we should output the Insights trace after running a test so that we know where to grab it from when we're done
 		/// </summary>
diff --git a/Engine/Plugins/Performance/AutomatedPerfTesting/Build/Scripts/AutomatedPerfTestNode.cs b/Engine/Plugins/Performance/AutomatedPerfTesting/Build/Scripts/AutomatedPerfTestNode.cs
--- a/Engine/Plugins/Performance/AutomatedPerfTesting/Build/Scripts/AutomatedPerfTestNode.cs
+++ b/Engine/Plugins/Performance/AutomatedPerfTesting/Build/Scripts/AutomatedPerfTestNode.cs
@@ -240,15 +240,17 @@ namespace AutomatedPerfTest
 					// to preserve them until we're ready to make our report on the final iteration.
 					CopyPerfFilesToTempDir(ArtifactPath);
 
-					// Local report generation is an example of how to use the PerfReportTool.
-					if (!Globals.Params.ParseParam("NoLocalReports"))
+					// Local report generation is useful for people conducting tests locally without a centralized server (for A/B testing for instance).
+					// To make it work out of the box, enable by default for non-build machine runs
+					if (Globals.Params.ParseParam("LocalReports") || (!CommandUtils.IsBuildMachine && !Globals.Params.ParseParam("NoLocalReports")))
 					{
 						// NOTE: This does not currently work with long paths due to the CsvTools not properly supporting them.
 						Log.Info("Generating performance reports using PerfReportTool.");
 						GenerateLocalPerfReport(Context.GetRoleContext(UnrealTargetRole.Client).Platform, ArtifactPath);
 					}
 
-					if (Globals.Params.ParseParam("PerfReportServer") &&
+					// On build machines, default to producing PRS report
+					if ((CommandUtils.IsBuildMachine || Globals.Params.ParseParam("PerfReportServer")) &&
 						!Globals.Params.ParseParam("SkipPerfReportServer"))
 					{
 						Dictionary<string, dynamic> CommonDataSourceFields = new Dictionary<string, dynamic>
@@ -349,7 +351,7 @@ namespace AutomatedPerfTest
 			}
 			
 			var ReportConfigDir = Path.Combine(Context.Options.ProjectPath.Directory.FullName, "Build", "Scripts", "PerfReport");
-			var ReportPath = Path.Combine(ReportCacheDir, "Reports", "Performance");
+			var ReportPath = Path.Combine(ArtifactPath, "Reports", "Performance");
 
 			// Csv files may have been output in one of two places.
 			// Check both...
@@ -519,9 +521,31 @@ namespace AutomatedPerfTest
 			}
 		}
 		
-		protected virtual string GetSubtestName()
+		protected virtual string GetSubtestName(TConfigClass Config)
 		{
-			return "Performance";
+			// Options like DoLLM and DoInsightsTrace are heavy enough to be in their own subtest type
+			// Since they are not exclusive, create a yet another subtest type if both are specified
+			if (Config.DoInsightsTrace && Config.DoLLM)
+			{
+				throw new AutomationException($"Running Insights trace with LLM is not a practical test due to LLM's large overhead.");
+			}
+
+			if (Config.DoInsightsTrace)
+			{
+				if (Config.TraceChannels != "default,screenshot,stats")
+				{
+					throw new AutomationException($"Running Insights trace with non-default channels {Config.TraceChannels}, new subtest type is needed to avoid contaminating 'Insights' subtest results.");
+				}
+
+				return "Insights";
+			}
+
+			if (Config.DoLLM)
+			{
+				return "LLM";
+			}
+
+			return "Perf";
 		}
 
 		public override TConfigClass GetConfiguration()
@@ -529,6 +553,12 @@ namespace AutomatedPerfTest
 			TConfigClass Config = base.GetConfiguration();
 			Config.MaxDuration = Context.TestParams.ParseValue("MaxDuration", 60 * 60);  // 1 hour max
 
+			// set PerfCacheRoot to some sane value
+			if (string.IsNullOrEmpty(Config.PerfCacheRoot))
+			{
+				Config.PerfCacheRoot = Path.Combine(Unreal.RootDirectory.FullName, "GauntletTemp", "PerfCache");
+			}
+
 			UnrealTestRole ClientRole = Config.RequireRole(UnrealTargetRole.Client);
 			// the controller will be added by the subclasses
 
@@ -536,7 +566,8 @@ namespace AutomatedPerfTest
 			
 			ClientRole.CommandLineParams.Add("-deterministic");
 
-			Log.Info("AutomatedPerfTestNode<>.GetConfiguration(): Config.DoFPSChart={0}, Config.DoCSVProfiler={1}, Config.DoVideoCapture={2}, Config.DoInsightsTrace={3}", Config.DoFPSChart, Config.DoCSVProfiler, Config.DoVideoCapture, Config.DoInsightsTrace);
+			Log.Info("AutomatedPerfTestNode<>.GetConfiguration(): Config.DoFPSChart={0}, Config.DoCSVProfiler={1}, Config.DoVideoCapture={2}, Config.DoInsightsTrace={3}, Config.DoLLM={4}", 
+				Config.DoFPSChart, Config.DoCSVProfiler, Config.DoVideoCapture, Config.DoInsightsTrace, Config.DoLLM);
 			
 			ClientRole.CommandLineParams.AddOrAppendParamValue("AutomatedPerfTest.TestName", Config.TestName);
 
@@ -554,6 +585,12 @@ namespace AutomatedPerfTest
 				}
 			}
 
+			if (Config.DoLLM)
+			{
+				ClientRole.CommandLineParams.Add("llm");
+				ClientRole.CommandLineParams.Add("llmcsv");
+			}
+
 			if (Config.DoFPSChart)
 			{
 				ClientRole.CommandLineParams.Add("AutomatedPerfTest.DoFPSChart");
@@ -569,7 +606,7 @@ namespace AutomatedPerfTest
 				{
 					string.Format("testname={0}", Context.BuildInfo.ProjectName),
 					"gauntletTestType=AutomatedPerfTest",
-					string.Format("gauntletSubTest={0}", GetSubtestName()),
+					string.Format("gauntletSubTest={0}", GetSubtestName(Config)),
 					"testBuildIsPreflight=" + (ReportGenUtils.IsTestingPreflightBuild(OriginalBuildName) ? "1" : "0"),
 					"testBuildVersion=" + OriginalBuildName
 				};
@@ -732,7 +769,18 @@ namespace AutomatedPerfTest
 			return Config;
 		}
 	}
-	
+
+	/// <summary>
+	/// "Standard issue" implementation usable for samples that don't need anything more advanced
+	/// </summary>
+	public class SequenceTest : AutomatedSequencePerfTestNode<AutomatedSequencePerfTestConfig>
+	{
+		public SequenceTest(Gauntlet.UnrealTestContext InContext)
+			: base(InContext)
+		{
+		}
+	}
+
 	/// <summary>
 	/// Implementation of a Gauntlet TestNode for AutomatedPerfTest plugin
 	/// </summary>
@@ -762,4 +810,15 @@ namespace AutomatedPerfTest
 			return Config;
 		}
 	}
+
+	/// <summary>
+	/// "Standard issue" implementation usable for samples that don't need anything more advanced
+	/// </summary>
+	public class MaterialTest : AutomatedMaterialPerfTestNode<AutomatedMaterialPerfTestConfig>
+	{
+		public MaterialTest(Gauntlet.UnrealTestContext InContext)
+			: base(InContext)
+		{
+		}
+	}
 }
