AutomatedPerfTest: fix boolean parameters.

diff --git a/Engine/Plugins/Performance/AutomatedPerfTesting/Build/Inc/AutomatedPerfTestCommonSettings.xml b/Engine/Plugins/Performance/AutomatedPerfTesting/Build/Inc/AutomatedPerfTestCommonSettings.xml
--- a/Engine/Plugins/Performance/AutomatedPerfTesting/Build/Inc/AutomatedPerfTestCommonSettings.xml
+++ b/Engine/Plugins/Performance/AutomatedPerfTesting/Build/Inc/AutomatedPerfTestCommonSettings.xml
@@ -37,7 +37,6 @@
 	<Property Name="AutomatedPerformanceCommonArgs" Value="$(AutomatedPerformanceCommonArgs) -AutomatedPerfTest.TestName=$(TestID)"/>
 	<Property Name="AutomatedPerformanceCommonArgs" Value="$(AutomatedPerformanceCommonArgs) -AutomatedPerfTest.DeviceProfileOverride=$(DeviceProfileOverride)" If="'$(DeviceProfileOverride)' != ''"/>
 	<Property Name="AutomatedPerformanceCommonArgs" Value="$(AutomatedPerformanceCommonArgs) -AutomatedPerfTest.TraceChannels=$(TraceChannels)" If="$(DoInsightsTrace) and '$(TraceChannels)' != ''"/>
-	<Property Name="AutomatedPerformanceCommonArgs" Value="$(AutomatedPerformanceCommonArgs) -AutomatedPerfTest.DoCSVProfiler" If="$(DoCSVProfiler)"/>
 	<Property Name="AutomatedPerformanceCommonArgs" Value="$(AutomatedPerformanceCommonArgs) -AutomatedPerfTest.DoFPSChart" If="$(DoFPSChart)"/>
 	<!--  disabled until we find a solution for escaping the path names in JSON nodes
 	<Property Name="AutomatedPerformanceCommonArgs" Value="$(AutomatedPerformanceCommonArgs) -AutomatedPerfTest.PerfCacheRoot=$(PerfCacheRoot)"/>
@@ -50,16 +49,16 @@
 	<!-- append plugin-provided tests to GauntletSettings' TargetTestList -->
 	<Do If="$(WithAPTSequenceTests)">
 		<ForEach Name="MapSequenceComboName" Values="$(MapSequenceComboNames)" Separator="+">
-			<Property Name="TargetTestList" Value="$(TargetTestList)AutomatedPerfTest.SequenceTest{-AutomatedPerfTest.SequencePerfTest.MapSequenceName=$(MapSequenceComboName) $(AutomatedPerformanceCommonArgs)}[Sequence Perf $(MapSequenceComboName)]+" If="$(DoPerf)" />
-			<Property Name="TargetTestList" Value="$(TargetTestList)AutomatedPerfTest.SequenceTest{-AutomatedPerfTest.SequencePerfTest.MapSequenceName=$(MapSequenceComboName) -AutomatedPerfTest.DoLLM=true $(AutomatedPerformanceCommonArgs)}[Sequence LLM $(MapSequenceComboName)]+" If="$(DoLLM)"/>
-			<Property Name="TargetTestList" Value="$(TargetTestList)AutomatedPerfTest.SequenceTest{-AutomatedPerfTest.SequencePerfTest.MapSequenceName=$(MapSequenceComboName) -AutomatedPerfTest.DoInsightsTrace=true $(AutomatedPerformanceCommonArgs)}[Sequence Insights $(MapSequenceComboName)]+" If="$(DoInsightsTrace)"/>
+			<Property Name="TargetTestList" Value="$(TargetTestList)AutomatedPerfTest.SequenceTest{-AutomatedPerfTest.SequencePerfTest.MapSequenceName=$(MapSequenceComboName) -AutomatedPerfTest.DoCSVProfiler $(AutomatedPerformanceCommonArgs)}[Sequence Perf $(MapSequenceComboName)]+" If="$(DoPerf)" />
+			<Property Name="TargetTestList" Value="$(TargetTestList)AutomatedPerfTest.SequenceTest{-AutomatedPerfTest.SequencePerfTest.MapSequenceName=$(MapSequenceComboName) -AutomatedPerfTest.DoCSVProfiler -AutomatedPerfTest.DoLLM $(AutomatedPerformanceCommonArgs)}[Sequence LLM $(MapSequenceComboName)]+" If="$(DoLLM)"/>
+			<Property Name="TargetTestList" Value="$(TargetTestList)AutomatedPerfTest.SequenceTest{-AutomatedPerfTest.SequencePerfTest.MapSequenceName=$(MapSequenceComboName) -AutomatedPerfTest.DoCSVProfiler -AutomatedPerfTest.DoInsightsTrace $(AutomatedPerformanceCommonArgs)}[Sequence Insights $(MapSequenceComboName)]+" If="$(DoInsightsTrace)"/>
 		</ForEach>
 	</Do>
 
 	<Do If="$(WithAPTMaterialTest)">
-		<Property Name="TargetTestList" Value="$(TargetTestList)AutomatedPerfTest.MaterialTest($(AutomatedPerformanceCommonArgs))[Material Perf]+" If="$(DoPerf)" />
-		<Property Name="TargetTestList" Value="$(TargetTestList)AutomatedPerfTest.MaterialTest(AutomatedPerfTest.DoLLM=true $(AutomatedPerformanceCommonArgs))[Material LLM]+" If="$(DoLLM)"/>
-		<Property Name="TargetTestList" Value="$(TargetTestList)AutomatedPerfTest.MaterialTest(AutomatedPerfTest.DoInsightsTrace=true $(AutomatedPerformanceCommonArgs))[Material Insights]+" If="$(DoInsightsTrace)"/>
+		<Property Name="TargetTestList" Value="$(TargetTestList)AutomatedPerfTest.MaterialTest{-AutomatedPerfTest.DoCSVProfiler $(AutomatedPerformanceCommonArgs)}[Material Perf]+" If="$(DoPerf)" />
+		<Property Name="TargetTestList" Value="$(TargetTestList)AutomatedPerfTest.MaterialTest{-AutomatedPerfTest.DoCSVProfiler -AutomatedPerfTest.DoLLM $(AutomatedPerformanceCommonArgs)}[Material LLM]+" If="$(DoLLM)"/>
+		<Property Name="TargetTestList" Value="$(TargetTestList)AutomatedPerfTest.MaterialTest{-AutomatedPerfTest.DoCSVProfiler -AutomatedPerfTest.DoInsightsTrace $(AutomatedPerformanceCommonArgs)}[Material Insights]+" If="$(DoInsightsTrace)"/>
 	</Do>
 
 
diff --git a/Engine/Plugins/Performance/AutomatedPerfTesting/Build/Scripts/AutomatedPerfTestConfig.cs b/Engine/Plugins/Performance/AutomatedPerfTesting/Build/Scripts/AutomatedPerfTestConfig.cs
--- a/Engine/Plugins/Performance/AutomatedPerfTesting/Build/Scripts/AutomatedPerfTestConfig.cs
+++ b/Engine/Plugins/Performance/AutomatedPerfTesting/Build/Scripts/AutomatedPerfTestConfig.cs
@@ -54,7 +54,7 @@ namespace AutomatedPerfTest
 		/// <summary>
 		/// If we're running on the build machine
 		/// </summary>
-		[AutoParamWithNames(true, "AutomatedPerfTest.DoCSVProfiler")]
+		[AutoParamWithNames(false, "AutomatedPerfTest.DoCSVProfiler")]
 		public bool DoCSVProfiler;
 		
 		/// <summary>
@@ -91,7 +91,7 @@ namespace AutomatedPerfTest
 		/// If true, will look for the shipping configuration of Unreal Insights in order to parse Insights trace files
 		/// This should be True unless you need to test an issue in parsing the Insights file itself.
 		/// </summary>
-		[AutoParamWithNames(true, "AutomatedPerfTest.UseShippingInsights")]
+		[AutoParamWithNames(false, "AutomatedPerfTest.UseShippingInsights")]
 		public bool UseShippingInsights;
 		
 		public string DataSourceName;
