AutomatedPerfTesting - Metadata Update
+ Adding additional CSVProfiler metadata for single, separate, and granular output modes
* Changing StaticCamera to default to Granular CSV output
* Changing the TestName cvar to TestID since TestName is reserved for other systems
+ Adding Project and Common settings for StaticCamera tests to plugin buildgraph
+ Adding "SummaryTable" var to APT test nodes which is passed to PRT.exe when generating reports for tests
* Parameterizing ResX and ResY, defaults to 4k
- removing extraneous test buildgraphs and automation scripts

diff --git a/Engine/Plugins/Performance/AutomatedPerfTesting/Build/Inc/AutomatedMaterialPerfTestMacro.xml b/Engine/Plugins/Performance/AutomatedPerfTesting/Build/Inc/AutomatedMaterialPerfTestMacro.xml
deleted file mode 100644
diff --git a/Engine/Plugins/Performance/AutomatedPerfTesting/Build/Inc/AutomatedPerfTestCommonSettings.xml b/Engine/Plugins/Performance/AutomatedPerfTesting/Build/Inc/AutomatedPerfTestCommonSettings.xml
--- a/Engine/Plugins/Performance/AutomatedPerfTesting/Build/Inc/AutomatedPerfTestCommonSettings.xml
+++ b/Engine/Plugins/Performance/AutomatedPerfTesting/Build/Inc/AutomatedPerfTestCommonSettings.xml
@@ -22,9 +22,12 @@
 	<Option Name="DataSourceName" DefaultValue="Automation.$(ProjectName).Performance" Description="To which PRS data source should this be uploaded"/>
 	<Option Name="PerfCacheRoot" DefaultValue="$(RootDir)/$(ProjectPath)/PerfCache" Description="Path to the perf cache where all the output data is stored. Defaults to a path in the local project's directory. Override to a common network-accessible path when running on build machines"/>
 	<Option Name="GenerateLocalReports" DefaultValue="false" Restrict="true|false" Description="Whether or not to create a local PerfReportTool report" />
+	<Option Name="ResX" DefaultValue="3840" Restrict="^[0-9]+$" Description="Resolution width to run tests with" />
+	<Option Name="ResY" DefaultValue="2160" Restrict="^[0-9]+$" Description="Resolution height to run tests with" />
 
 	<!-- Plugin-provided tests -->
 	<Option Name="WithAPTSequenceTests" DefaultValue="true" Restrict="true|false" Description="Whether to run Auto Perf Test sequence flythrough test." />
+	<Option Name="WithAPTStaticCameraTests" DefaultValue="false" Restrict="true|false" Description="Whether to run Auto Perf Test static camera tests."  />
 	<Option Name="WithAPTMaterialTest" DefaultValue="false" Restrict="true|false" Description="Whether to run Auto Perf Test material test."  />
 	<Option Name="WithAPTReplayTest"   DefaultValue="false" Restrict="true|false" Description="Whether to run Auto Perf Test replay test."  />
 
@@ -35,7 +38,7 @@
 
 	<!-- Construct the common arguments for the performance test and report commands -->
 	<Property Name="AutomatedPerformanceCommonArgs" Value = "-AutomatedPerfTest.UseShippingInsights=$(UseShippingInsights)" />
-	<Property Name="AutomatedPerformanceCommonArgs" Value="$(AutomatedPerformanceCommonArgs) -AutomatedPerfTest.TestName=$(TestID)"/>
+	<Property Name="AutomatedPerformanceCommonArgs" Value="$(AutomatedPerformanceCommonArgs) -AutomatedPerfTest.TestID=$(TestID)"/>
 	<Property Name="AutomatedPerformanceCommonArgs" Value="$(AutomatedPerformanceCommonArgs) -AutomatedPerfTest.DeviceProfileOverride=$(DeviceProfileOverride)" If="'$(DeviceProfileOverride)' != ''"/>
 	<Property Name="AutomatedPerformanceCommonArgs" Value="$(AutomatedPerformanceCommonArgs) -AutomatedPerfTest.TraceChannels=$(TraceChannels)" If="$(DoInsightsTrace) and '$(TraceChannels)' != ''"/>
 	<Property Name="AutomatedPerformanceCommonArgs" Value="$(AutomatedPerformanceCommonArgs) -AutomatedPerfTest.DoFPSChart" If="$(DoFPSChart)"/>
@@ -44,7 +47,7 @@
 	-->
 	<Property Name="AutomatedPerformanceCommonArgs" Value="$(AutomatedPerformanceCommonArgs) -PerfReportServer" If="$(PerfReportServer) And $(DoCSVProfiler)"/>
 	<Property Name="AutomatedPerformanceCommonArgs" Value="$(AutomatedPerformanceCommonArgs) -NoLocalReports" If="!$(GenerateLocalReports) And $(DoCSVProfiler)" />
-	<Property Name="AutomatedPerformanceCommonArgs" Value="$(AutomatedPerformanceCommonArgs) -resX=3840 -resY=2160"/>
+	<Property Name="AutomatedPerformanceCommonArgs" Value="$(AutomatedPerformanceCommonArgs) -resX=$(ResX) -resY=$(ResY)"/>
 	<Property Name="AutomatedPerformanceCommonArgs" Value="$(AutomatedPerformanceCommonArgs) -AutomatedPerfTest.DoVideoCapture" If="$(DoVideoCapture)"/>
 
 	<!-- append plugin-provided tests to GauntletSettings' TargetTestList -->
@@ -56,6 +59,14 @@
 		</ForEach>
 	</Do>
 
+	<Do If="$(WithAPTStaticCameraTests)">
+		<ForEach Name="StaticCameraMapName" Values="$(StaticCameraMapNames)" Separator="+">
+			<Property Name="TargetTestList" Value="$(TargetTestList)AutomatedPerfTest.StaticCameraTest{-AutomatedPerfTest.StaticCameraPerfTest.MapName=$(StaticCameraMapName) -AutomatedPerfTest.DoCSVProfiler $(AutomatedPerformanceCommonArgs)}[Static Camera Perf $(StaticCameraMapName)]+" If="$(DoPerf)" />
+			<Property Name="TargetTestList" Value="$(TargetTestList)AutomatedPerfTest.StaticCameraTest{-AutomatedPerfTest.StaticCameraPerfTest.MapName=$(StaticCameraMapName) -AutomatedPerfTest.DoCSVProfiler -AutomatedPerfTest.DoLLM=true $(AutomatedPerformanceCommonArgs)}[Static Camera LLM $(StaticCameraMapName)]+" If="$(DoLLM)"/>
+			<Property Name="TargetTestList" Value="$(TargetTestList)AutomatedPerfTest.StaticCameraTest{-AutomatedPerfTest.StaticCameraPerfTest.MapName=$(StaticCameraMapName) -AutomatedPerfTest.DoCSVProfiler -AutomatedPerfTest.DoInsightsTrace=true $(AutomatedPerformanceCommonArgs)}[Static Camera Insights $(StaticCameraMapName)]+" If="$(DoInsightsTrace)"/>
+		</ForEach>
+	</Do>
+
 	<Do If="$(WithAPTMaterialTest)">
 		<Property Name="TargetTestList" Value="$(TargetTestList)AutomatedPerfTest.MaterialTest{-AutomatedPerfTest.DoCSVProfiler $(AutomatedPerformanceCommonArgs)}[Material Perf]+" If="$(DoPerf)" />
 		<Property Name="TargetTestList" Value="$(TargetTestList)AutomatedPerfTest.MaterialTest{-AutomatedPerfTest.DoCSVProfiler -AutomatedPerfTest.DoLLM $(AutomatedPerformanceCommonArgs)}[Material LLM]+" If="$(DoLLM)"/>
diff --git a/Engine/Plugins/Performance/AutomatedPerfTesting/Build/Inc/AutomatedPerfTestProjectSettings.xml b/Engine/Plugins/Performance/AutomatedPerfTesting/Build/Inc/AutomatedPerfTestProjectSettings.xml
--- a/Engine/Plugins/Performance/AutomatedPerfTesting/Build/Inc/AutomatedPerfTestProjectSettings.xml
+++ b/Engine/Plugins/Performance/AutomatedPerfTesting/Build/Inc/AutomatedPerfTestProjectSettings.xml
@@ -2,6 +2,9 @@
 <BuildGraph xmlns="http://www.epicgames.com/BuildGraph" xmlns:xsi="http://www.w3.org/2001/XMLSchema-instance" xsi:schemaLocation="http://www.epicgames.com/BuildGraph ../../../../../../../Engine/Build/Graph/Schema.xsd" >
 
 	<!-- These are options for Automated Perf Testing that can - and often should - be overriden in a given project -->
-	<Option Name="MapSequenceComboNames" DefaultValue="" Description="+ separated list of names, referencing the names of Map/Sequence combos defined in Automated Perf Test Project Settings" />
+	<Option Name="MapSequenceComboNames" DefaultValue="" Description="+ separated list of names, referencing the names of Map/Sequence combos defined in Automated Perf Test | Sequence Project Settings" />
+
+	<Option Name="StaticCameraMapNames" DefaultValue="" Description="+ separated list of names, referencing the names of Maps defined in Automated Perf Test | Static Camera Project Settings" />
+
 	<Option Name="ReplayName" DefaultValue="" Description="Name of replay to run when running Automated Perf Test Replay" />
 </BuildGraph>
diff --git a/Engine/Plugins/Performance/AutomatedPerfTesting/Build/Inc/AutomatedSequencePerfTestMacro.xml b/Engine/Plugins/Performance/AutomatedPerfTesting/Build/Inc/AutomatedSequencePerfTestMacro.xml
deleted file mode 100644
diff --git a/Engine/Plugins/Performance/AutomatedPerfTesting/Build/Inc/AutomatedStaticCameraPerfTestMacro.xml b/Engine/Plugins/Performance/AutomatedPerfTesting/Build/Inc/AutomatedStaticCameraPerfTestMacro.xml
deleted file mode 100644
diff --git a/Engine/Plugins/Performance/AutomatedPerfTesting/Build/Scripts/AutomatedPerfTestConfig.cs b/Engine/Plugins/Performance/AutomatedPerfTesting/Build/Scripts/AutomatedPerfTestConfig.cs
--- a/Engine/Plugins/Performance/AutomatedPerfTesting/Build/Scripts/AutomatedPerfTestConfig.cs
+++ b/Engine/Plugins/Performance/AutomatedPerfTesting/Build/Scripts/AutomatedPerfTestConfig.cs
@@ -30,8 +30,8 @@ namespace AutomatedPerfTest
 		/// <summary>
 		/// Name of the test, useful for identifying it later
 		/// </summary>
-		[AutoParamWithNames("AutomatedPerfTest.TestName")]
-		public string TestName = "";
+		[AutoParamWithNames("AutomatedPerfTest.TestID")]
+		public string TestID = "";
 		
 		/// <summary>
 		/// If set, will prepend platform name and use this device profile instead of the default
@@ -40,7 +40,7 @@ namespace AutomatedPerfTest
 		public string DeviceProfileOverride = "";
 		
 		/// <summary>
-		/// Whether to do an Insights trace
+		/// If data from the test should be gathered with Unreal Insights
 		/// </summary>
 		[AutoParamWithNames(false, "AutomatedPerfTest.DoInsightsTrace")]
 		public bool DoInsightsTrace;
@@ -52,13 +52,13 @@ namespace AutomatedPerfTest
 		public string TraceChannels = "default,screenshot,stats";
 		
 		/// <summary>
-		/// If we're running on the build machine
+		/// If data from the test should be gathered with the CSV Profiler
 		/// </summary>
 		[AutoParamWithNames(false, "AutomatedPerfTest.DoCSVProfiler")]
 		public bool DoCSVProfiler;
 		
 		/// <summary>
-		/// If we're running on the build machine
+		/// If data from the test should be gathered with FPS Charts
 		/// </summary>
 		[AutoParamWithNames(false, "AutomatedPerfTest.DoFPSChart")]
 		public bool DoFPSChart;
diff --git a/Engine/Plugins/Performance/AutomatedPerfTesting/Build/Scripts/AutomatedPerfTestInsightsReport.cs b/Engine/Plugins/Performance/AutomatedPerfTesting/Build/Scripts/AutomatedPerfTestInsightsReport.cs
deleted file mode 100644
diff --git a/Engine/Plugins/Performance/AutomatedPerfTesting/Build/Scripts/AutomatedPerfTestNode.cs b/Engine/Plugins/Performance/AutomatedPerfTesting/Build/Scripts/AutomatedPerfTestNode.cs
--- a/Engine/Plugins/Performance/AutomatedPerfTesting/Build/Scripts/AutomatedPerfTestNode.cs
+++ b/Engine/Plugins/Performance/AutomatedPerfTesting/Build/Scripts/AutomatedPerfTestNode.cs
@@ -24,6 +24,7 @@ namespace AutomatedPerfTest
 	public abstract class AutomatedPerfTestNode<TConfigClass> : UnrealTestNode<TConfigClass>
 		where TConfigClass : AutomatedPerfTestConfigBase, new()
 	{
+		public string SummaryTable = "default";
 		public AutomatedPerfTestNode(UnrealTestContext InContext) : base(InContext)
 		{
 			// We need to save off the build name as if this is a preflight that suffix will be stripped
@@ -416,7 +417,7 @@ namespace AutomatedPerfTest
 			string SearchPattern = $"{Context.BuildInfo.ProjectName}*";
 			
 			// Produce the detailed report, and update the perf cache
-			CommandUtils.RunAndLog(ToolPath.FullName, $"-csvdir \"{NewestDir}\" -o \"{ReportPath}\" -reportxmlbasedir \"{ReportConfigDir}\" -summaryTableCache \"{ReportCacheDir}\" -searchpattern {SearchPattern} -metadatafilter platform=\"{PlatformNameFilter}\"", out int ErrorCode);
+			CommandUtils.RunAndLog(ToolPath.FullName, $"-csvdir \"{NewestDir}\" -o \"{ReportPath}\" -reportxmlbasedir \"{ReportConfigDir}\" -summaryTable \"{SummaryTable}\" -summaryTableCache \"{ReportCacheDir}\" -searchpattern {SearchPattern} -metadatafilter platform=\"{PlatformNameFilter}\"", out int ErrorCode);
 			if (ErrorCode != 0)
 			{
 				Logger.LogError("PerfReportTool returned error code \"{ErrorCode}\" while generating detailed report.", ErrorCode);
@@ -569,7 +570,7 @@ namespace AutomatedPerfTest
 			Log.Info("AutomatedPerfTestNode<>.GetConfiguration(): Config.DoFPSChart={0}, Config.DoCSVProfiler={1}, Config.DoVideoCapture={2}, Config.DoInsightsTrace={3}, Config.DoLLM={4}", 
 				Config.DoFPSChart, Config.DoCSVProfiler, Config.DoVideoCapture, Config.DoInsightsTrace, Config.DoLLM);
 			
-			ClientRole.CommandLineParams.AddOrAppendParamValue("AutomatedPerfTest.TestName", Config.TestName);
+			ClientRole.CommandLineParams.AddOrAppendParamValue("AutomatedPerfTest.TestID", Config.TestID);
 
 			if (Config.DeviceProfileOverride != String.Empty)
 			{
@@ -602,14 +603,12 @@ namespace AutomatedPerfTest
 				ClientRole.CommandLineParams.Add("csvGpuStats");
 
 				// Add CSV metadata
-				List<string> CsvMetadata = new List<string>
-				{
-					string.Format("testname={0}", Context.BuildInfo.ProjectName),
-					"gauntletTestType=AutomatedPerfTest",
-					string.Format("gauntletSubTest={0}", GetSubtestName(Config)),
-					"testBuildIsPreflight=" + (ReportGenUtils.IsTestingPreflightBuild(OriginalBuildName) ? "1" : "0"),
-					"testBuildVersion=" + OriginalBuildName
-				};
+				List<string> CsvMetadata = new List<string> { };
+				CsvMetadata.Add($"testname={Context.BuildInfo.ProjectName}");
+				CsvMetadata.Add("gauntletTestType=AutomatedPerfTest");
+				CsvMetadata.Add($"gauntletSubTest={GetSubtestName(Config)}");
+				CsvMetadata.Add("testBuildIsPreflight=" + (ReportGenUtils.IsTestingPreflightBuild(OriginalBuildName) ? "1" : "0"));
+				CsvMetadata.Add($"testBuildVersion={OriginalBuildName}");
 
 				if (!string.IsNullOrEmpty(Context.BuildInfo.Branch) && Context.BuildInfo.Changelist != 0)
 				{
@@ -703,6 +702,7 @@ namespace AutomatedPerfTest
 	{
 		public AutomatedSequencePerfTestNode(UnrealTestContext InContext) : base(InContext)
 		{
+			SummaryTable = "sequence";
 		}
 
 		public override TConfigClass GetConfiguration()
@@ -774,6 +774,7 @@ namespace AutomatedPerfTest
 	{
 		public AutomatedStaticCameraPerfTestNode(UnrealTestContext InContext) : base(InContext)
 		{
+			SummaryTable = "staticcamera";
 		}
 
 		public override TConfigClass GetConfiguration()
@@ -813,6 +814,17 @@ namespace AutomatedPerfTest
 		{
 		}
 	}
+	
+	/// <summary>
+	/// "Standard issue" implementation usable for samples that don't need anything more advanced
+	/// </summary>
+	public class StaticCameraTest : AutomatedStaticCameraPerfTestNode<AutomatedStaticCameraPerfTestConfig>
+	{
+		public StaticCameraTest(Gauntlet.UnrealTestContext InContext)
+			: base(InContext)
+		{
+		}
+	}
 
 	/// <summary>
 	/// Implementation of a Gauntlet TestNode for AutomatedPerfTest plugin
@@ -823,6 +835,7 @@ namespace AutomatedPerfTest
 	{
 		public AutomatedMaterialPerfTestNode(UnrealTestContext InContext) : base(InContext)
 		{
+			SummaryTable = "materials";
 		}
 
 		public override TConfigClass GetConfiguration()
@@ -849,6 +862,7 @@ namespace AutomatedPerfTest
 	/// </summary>
 	public class MaterialTest : AutomatedMaterialPerfTestNode<AutomatedMaterialPerfTestConfig>
 	{
+		
 		public MaterialTest(Gauntlet.UnrealTestContext InContext)
 			: base(InContext)
 		{
diff --git a/Engine/Plugins/Performance/AutomatedPerfTesting/Source/AutomatedPerfTesting/Private/AutomatedMaterialPerfTest.cpp b/Engine/Plugins/Performance/AutomatedPerfTesting/Source/AutomatedPerfTesting/Private/AutomatedMaterialPerfTest.cpp
--- a/Engine/Plugins/Performance/AutomatedPerfTesting/Source/AutomatedPerfTesting/Private/AutomatedMaterialPerfTest.cpp
+++ b/Engine/Plugins/Performance/AutomatedPerfTesting/Source/AutomatedPerfTesting/Private/AutomatedMaterialPerfTest.cpp
@@ -99,6 +99,21 @@ FString UAutomatedMaterialPerfTest::GetTestID()
 	return Super::GetTestID() + "_Materials";
 }
 
+bool UAutomatedMaterialPerfTest::TryStartCSVProfiler(FString CSVFileName)
+{
+	if(FCsvProfiler* const CsvProfiler = FCsvProfiler::Get())
+	{
+		if (GetCSVOutputMode() == EAutomatedPerfTestCSVOutputMode::Granular)
+		{
+			CsvProfiler->SetMetadata(TEXT("Material"), *GetCurrentMaterialRegionName());
+		}
+		return Super::TryStartCSVProfiler(CSVFileName);
+	}
+	
+	UE_LOG(LogAutomatedPerfTest, Warning, TEXT("CSVProfiler Start requested, but not available."))
+	return false;
+}
+
 void UAutomatedMaterialPerfTest::SetUpNextMaterial()
 {
 	CurrentMaterialIndex += 1;
diff --git a/Engine/Plugins/Performance/AutomatedPerfTesting/Source/AutomatedPerfTesting/Private/AutomatedPerfTestControllerBase.cpp b/Engine/Plugins/Performance/AutomatedPerfTesting/Source/AutomatedPerfTesting/Private/AutomatedPerfTestControllerBase.cpp
--- a/Engine/Plugins/Performance/AutomatedPerfTesting/Source/AutomatedPerfTesting/Private/AutomatedPerfTestControllerBase.cpp
+++ b/Engine/Plugins/Performance/AutomatedPerfTesting/Source/AutomatedPerfTesting/Private/AutomatedPerfTestControllerBase.cpp
@@ -19,7 +19,7 @@
 #include "Misc/Paths.h"
 #include "ProfilingDebugging/TraceScreenshot.h"
 #include "UnrealClient.h"
-
+#include "UnrealEngine.h"
 
 DEFINE_LOG_CATEGORY(LogAutomatedPerfTest)
 CSV_DEFINE_CATEGORY(AutomatedPerfTest, true);
@@ -97,12 +97,12 @@ UAutomatedPerfTestControllerBase::UAutomatedPerfTestControllerBase(const FObject
     , CSVOutputMode(EAutomatedPerfTestCSVOutputMode::Single)
 {
 	// cache this off once, so that it's consistent throughout a session
-	TestDatetime = FDateTime::Now().ToString(TEXT("%Y%m%d-%H%M%S")); 
+	TestDatetime = FDateTime::Now().ToString(TEXT("%Y%m%d-%H%M%S"));
 }
 
 FString UAutomatedPerfTestControllerBase::GetTestName()
 {
-	return TestName;
+	return TestID;
 }
 
 FString UAutomatedPerfTestControllerBase::GetDeviceProfile()
@@ -187,6 +187,11 @@ bool UAutomatedPerfTestControllerBase::TryStartCSVProfiler(FString CSVFileName)
 		}
 
 		UE_LOG(LogAutomatedPerfTest, Log, TEXT("Attempting to start CSV Profile to file %s"), *CSVFileName);
+
+		CsvProfiler->SetMetadata(TEXT("TestID"), *TestID);
+		CsvProfiler->SetMetadata(TEXT("Datetime"), *TestDatetime);
+		CsvProfiler->SetMetadata(TEXT("ResX"), *FString::FromInt(GSystemResolution.ResX));
+		CsvProfiler->SetMetadata(TEXT("ResY"), *FString::FromInt(GSystemResolution.ResY));
 		
 		CsvProfiler->BeginCapture(-1, FString(), CSVFileName);
 		CsvProfiler->SetDeviceProfileName(GetDeviceProfile());
@@ -490,7 +495,7 @@ void UAutomatedPerfTestControllerBase::OnInit()
 	FParse::Value(FCommandLine::Get(), TEXT("AutomatedPerfTest.TraceChannels="), TraceChannels, false);
 	
 	FParse::Value(FCommandLine::Get(), TEXT("AutomatedPerfTest.DeviceProfileOverride="), DeviceProfileOverride);
-	FParse::Value(FCommandLine::Get(), TEXT("AutomatedPerfTest.TestName="), TestName);
+	FParse::Value(FCommandLine::Get(), TEXT("AutomatedPerfTest.TestID="), TestID);
 	
 	if(FParse::Param(FCommandLine::Get(), TEXT("AutomatedPerfTest.DoInsightsTrace")))
 	{
diff --git a/Engine/Plugins/Performance/AutomatedPerfTesting/Source/AutomatedPerfTesting/Private/AutomatedSequencePerfTest.cpp b/Engine/Plugins/Performance/AutomatedPerfTesting/Source/AutomatedPerfTesting/Private/AutomatedSequencePerfTest.cpp
--- a/Engine/Plugins/Performance/AutomatedPerfTesting/Source/AutomatedPerfTesting/Private/AutomatedSequencePerfTest.cpp
+++ b/Engine/Plugins/Performance/AutomatedPerfTesting/Source/AutomatedPerfTesting/Private/AutomatedSequencePerfTest.cpp
@@ -187,6 +187,25 @@ void UAutomatedSequencePerfTest::Exit()
 	Super::Exit();
 }
 
+bool UAutomatedSequencePerfTest::TryStartCSVProfiler(FString CSVFileName)
+{
+	if(FCsvProfiler* const CsvProfiler = FCsvProfiler::Get())
+	{
+		if (GetCSVOutputMode() != EAutomatedPerfTestCSVOutputMode::Single)
+		{
+			CsvProfiler->SetMetadata(TEXT("MapSequenceComboName"), *CurrentMapSequenceCombo->ComboName.ToString());
+		}
+		if (GetCSVOutputMode() == EAutomatedPerfTestCSVOutputMode::Granular)
+		{
+			CsvProfiler->SetMetadata(TEXT("CameraCut"), *GetCameraCutID());
+		}
+		return Super::TryStartCSVProfiler(CSVFileName);
+	}
+	
+	UE_LOG(LogAutomatedPerfTest, Warning, TEXT("CSVProfiler Start requested, but not available."))
+	return false;
+}
+
 void UAutomatedSequencePerfTest::OnCameraCut(UCameraComponent* CameraComponent)
 {
 	// null check the CurrentCamera so that we can use OnCameraCut to mark the starting camera cut region
@@ -243,7 +262,7 @@ FString UAutomatedSequencePerfTest::GetCameraCutID()
 	if(CurrentCamera != nullptr)
 	{
 		// getting the label of a spawnable camera from Sequencer in a packaged build isn't possible
-		// via this method. TODO: find a more reliable way to return camera cut names that are set in Sequencerj
+		// via this method. TODO: find a more reliable way to return camera cut names that are set in Sequencer
 		//const AActor* Owner = CurrentCamera->GetOwner();
 		//FString CutName = Owner ? Owner->GetActorNameOrLabel() : CurrentCamera->GetName();
 		//return CutName;
diff --git a/Engine/Plugins/Performance/AutomatedPerfTesting/Source/AutomatedPerfTesting/Private/StaticCameraTests/AutomatedStaticCameraPerfTestBase.cpp b/Engine/Plugins/Performance/AutomatedPerfTesting/Source/AutomatedPerfTesting/Private/StaticCameraTests/AutomatedStaticCameraPerfTestBase.cpp
--- a/Engine/Plugins/Performance/AutomatedPerfTesting/Source/AutomatedPerfTesting/Private/StaticCameraTests/AutomatedStaticCameraPerfTestBase.cpp
+++ b/Engine/Plugins/Performance/AutomatedPerfTesting/Source/AutomatedPerfTesting/Private/StaticCameraTests/AutomatedStaticCameraPerfTestBase.cpp
@@ -13,7 +13,7 @@ UAutomatedStaticCameraPerfTestProjectSettings::UAutomatedStaticCameraPerfTestPro
 	, WarmUpTime(5.0)
 	, SoakTime(5.0)
 	, CooldownTime(1.0)
-	, CSVOutputMode(EAutomatedPerfTestCSVOutputMode::Separate)
+	, CSVOutputMode(EAutomatedPerfTestCSVOutputMode::Granular)
 {
 
 }
@@ -96,6 +96,25 @@ FString UAutomatedStaticCameraPerfTestBase::GetTestID()
 	return Super::GetTestID() + "_StaticCamera_" + GetCurrentMap();
 }
 
+bool UAutomatedStaticCameraPerfTestBase::TryStartCSVProfiler(FString CSVFileName)
+{
+	if(FCsvProfiler* const CsvProfiler = FCsvProfiler::Get())
+	{
+		if (GetCSVOutputMode() != EAutomatedPerfTestCSVOutputMode::Single)
+		{
+			CsvProfiler->SetMetadata(TEXT("MapName"), *CurrentMapName);
+		}
+		if (GetCSVOutputMode() == EAutomatedPerfTestCSVOutputMode::Granular)
+		{
+			CsvProfiler->SetMetadata(TEXT("CameraName"), *GetCurrentCameraRegionName());
+		}
+		return Super::TryStartCSVProfiler(CSVFileName);
+	}
+	
+	UE_LOG(LogAutomatedPerfTest, Warning, TEXT("CSVProfiler Start requested, but not available."))
+	return false;
+}
+
 void UAutomatedStaticCameraPerfTestBase::SetUpNextCamera()
 {
 	if(CamerasToTest.Num() <= 0)
diff --git a/Engine/Plugins/Performance/AutomatedPerfTesting/Source/AutomatedPerfTesting/Public/AutomatedMaterialPerfTest.h b/Engine/Plugins/Performance/AutomatedPerfTesting/Source/AutomatedPerfTesting/Public/AutomatedMaterialPerfTest.h
--- a/Engine/Plugins/Performance/AutomatedPerfTesting/Source/AutomatedPerfTesting/Public/AutomatedMaterialPerfTest.h
+++ b/Engine/Plugins/Performance/AutomatedPerfTesting/Source/AutomatedPerfTesting/Public/AutomatedMaterialPerfTest.h
@@ -112,6 +112,8 @@ public:
 	virtual void RunTest() override;
 
 	virtual FString GetTestID() override;
+
+	virtual bool TryStartCSVProfiler(FString CSVFileName) override;
 	// ~End UAutomatedPerfTestControllerBase Interface
 	
 	UFUNCTION()
diff --git a/Engine/Plugins/Performance/AutomatedPerfTesting/Source/AutomatedPerfTesting/Public/AutomatedPerfTestControllerBase.h b/Engine/Plugins/Performance/AutomatedPerfTesting/Source/AutomatedPerfTesting/Public/AutomatedPerfTestControllerBase.h
--- a/Engine/Plugins/Performance/AutomatedPerfTesting/Source/AutomatedPerfTesting/Public/AutomatedPerfTestControllerBase.h
+++ b/Engine/Plugins/Performance/AutomatedPerfTesting/Source/AutomatedPerfTesting/Public/AutomatedPerfTestControllerBase.h
@@ -85,8 +85,8 @@ public:
 	bool TryStartInsightsTrace();
 	bool TryStopInsightsTrace();
 
-	bool TryStartCSVProfiler();
-	bool TryStartCSVProfiler(FString CSVFileName);
+	virtual bool TryStartCSVProfiler();
+	virtual bool TryStartCSVProfiler(FString CSVFileName);
 	bool TryStopCSVProfiler();
 	
 	bool TryStartFPSChart();
@@ -131,7 +131,7 @@ protected:
 private:
 	FString TraceChannels;
 	FString TestDatetime;
-	FString TestName;
+	FString TestID;
 	FString DeviceProfileOverride;
 	bool bRequestsFPSChart;
 	bool bRequestsInsightsTrace;
diff --git a/Engine/Plugins/Performance/AutomatedPerfTesting/Source/AutomatedPerfTesting/Public/AutomatedSequencePerfTest.h b/Engine/Plugins/Performance/AutomatedPerfTesting/Source/AutomatedPerfTesting/Public/AutomatedSequencePerfTest.h
--- a/Engine/Plugins/Performance/AutomatedPerfTesting/Source/AutomatedPerfTesting/Public/AutomatedSequencePerfTest.h
+++ b/Engine/Plugins/Performance/AutomatedPerfTesting/Source/AutomatedPerfTesting/Public/AutomatedSequencePerfTest.h
@@ -96,6 +96,8 @@ public:
 	
 	virtual void TeardownTest(bool bExitAfterTeardown = true) override;
 	virtual void Exit() override;
+
+	virtual bool TryStartCSVProfiler(FString CSVFileName) override;
 	// ~End UAutomatedPerfTestControllerBase Interface
 	
 	// This function is called on world change to set up the map for the correct map/sequence combo
diff --git a/Engine/Plugins/Performance/AutomatedPerfTesting/Source/AutomatedPerfTesting/Public/StaticCameraTests/AutomatedStaticCameraPerfTestBase.h b/Engine/Plugins/Performance/AutomatedPerfTesting/Source/AutomatedPerfTesting/Public/StaticCameraTests/AutomatedStaticCameraPerfTestBase.h
--- a/Engine/Plugins/Performance/AutomatedPerfTesting/Source/AutomatedPerfTesting/Public/StaticCameraTests/AutomatedStaticCameraPerfTestBase.h
+++ b/Engine/Plugins/Performance/AutomatedPerfTesting/Source/AutomatedPerfTesting/Public/StaticCameraTests/AutomatedStaticCameraPerfTestBase.h
@@ -85,6 +85,8 @@ public:
 	virtual void RunTest() override;
 
 	virtual FString GetTestID() override;
+
+	virtual bool TryStartCSVProfiler(FString CSVFileName) override;
 	// ~End UAutomatedPerfTestControllerBase Interface
 	
 	UFUNCTION()
