Fix AutomatedPerfTesting runs to be repeated the desired number of times

* Each subtest type (Perf/LLM/Insights) will now repeat/iterate the configured number of times.
* Default values for Perf=6, LLM=3 and Insights Trace=1. This can be overriden if needed.

diff --git a/Engine/Plugins/Performance/AutomatedPerfTesting/Build/Inc/AutomatedPerfTestCommonSettings.xml b/Engine/Plugins/Performance/AutomatedPerfTesting/Build/Inc/AutomatedPerfTestCommonSettings.xml
--- a/Engine/Plugins/Performance/AutomatedPerfTesting/Build/Inc/AutomatedPerfTestCommonSettings.xml
+++ b/Engine/Plugins/Performance/AutomatedPerfTesting/Build/Inc/AutomatedPerfTestCommonSettings.xml
@@ -24,6 +24,9 @@
 	<Option Name="GenerateLocalReports" DefaultValue="false" Restrict="true|false" Description="Whether or not to create a local PerfReportTool report" />
 	<Option Name="ResX" DefaultValue="3840" Restrict="^[0-9]+$" Description="Resolution width to run tests with" />
 	<Option Name="ResY" DefaultValue="2160" Restrict="^[0-9]+$" Description="Resolution height to run tests with" />
+	<Option Name="IterationsPerf" DefaultValue="6" Restrict="^[0-9]+$" Description="Number of iterations to run perf test with" />
+	<Option Name="IterationsLLM" DefaultValue="3" Restrict="^[0-9]+$" Description="Number of iterations to run LLM test with" />
+	<Option Name="IterationsInsights" DefaultValue="1" Restrict="^[0-9]+$" Description="Number of iterations to run Insights test with" />
 
 	<Option Name="ExtraAutomatedPerformanceCommonArgs" DefaultValue="" Description="Additional Common arguments passed to Gauntlet command line"/>
 
@@ -53,33 +56,37 @@
 	<Property Name="AutomatedPerformanceCommonArgs" Value="$(AutomatedPerformanceCommonArgs) -AutomatedPerfTest.DoVideoCapture" If="$(DoVideoCapture)"/>
 	<Property Name="AutomatedPerformanceCommonArgs" Value="$(AutomatedPerformanceCommonArgs) $(ExtraAutomatedPerformanceCommonArgs)" If="'$(ExtraAutomatedPerformanceCommonArgs)' != ''"/>
 
+	<Property Name="PerfTestCommonArgs"     Value="$(AutomatedPerformanceCommonArgs) -iterations=$(IterationsPerf)"/>
+	<Property Name="LLMTestCommonArgs"      Value="$(AutomatedPerformanceCommonArgs) -iterations=$(IterationsLLM)"/>
+	<Property Name="InsightsTestCommonArgs" Value="$(AutomatedPerformanceCommonArgs) -iterations=$(IterationsInsights)"/>
+
 	<!-- append plugin-provided tests to GauntletSettings' TargetTestList -->
 	<Do If="$(WithAPTSequenceTests)">
 		<ForEach Name="MapSequenceComboName" Values="$(MapSequenceComboNames)" Separator="+">
-			<Property Name="TargetTestList" Value="$(TargetTestList)AutomatedPerfTest.SequenceTest{-AutomatedPerfTest.SequencePerfTest.MapSequenceName=$(MapSequenceComboName) -AutomatedPerfTest.DoCSVProfiler $(AutomatedPerformanceCommonArgs)}[Sequence Perf $(MapSequenceComboName)]+" If="$(DoPerf)" />
-			<Property Name="TargetTestList" Value="$(TargetTestList)AutomatedPerfTest.SequenceTest{-AutomatedPerfTest.SequencePerfTest.MapSequenceName=$(MapSequenceComboName) -AutomatedPerfTest.DoCSVProfiler -AutomatedPerfTest.DoLLM $(AutomatedPerformanceCommonArgs)}[Sequence LLM $(MapSequenceComboName)]+" If="$(DoLLM)"/>
-			<Property Name="TargetTestList" Value="$(TargetTestList)AutomatedPerfTest.SequenceTest{-AutomatedPerfTest.SequencePerfTest.MapSequenceName=$(MapSequenceComboName) -AutomatedPerfTest.DoCSVProfiler -AutomatedPerfTest.DoInsightsTrace $(AutomatedPerformanceCommonArgs)}[Sequence Insights $(MapSequenceComboName)]+" If="$(DoInsightsTrace)"/>
+			<Property Name="TargetTestList" Value="$(TargetTestList)AutomatedPerfTest.SequenceTest{-AutomatedPerfTest.SequencePerfTest.MapSequenceName=$(MapSequenceComboName) -AutomatedPerfTest.DoCSVProfiler $(PerfTestCommonArgs)}[Sequence Perf $(MapSequenceComboName)]+" If="$(DoPerf)" />
+			<Property Name="TargetTestList" Value="$(TargetTestList)AutomatedPerfTest.SequenceTest{-AutomatedPerfTest.SequencePerfTest.MapSequenceName=$(MapSequenceComboName) -AutomatedPerfTest.DoCSVProfiler -AutomatedPerfTest.DoLLM $(LLMTestCommonArgs)}[Sequence LLM $(MapSequenceComboName)]+" If="$(DoLLM)"/>
+			<Property Name="TargetTestList" Value="$(TargetTestList)AutomatedPerfTest.SequenceTest{-AutomatedPerfTest.SequencePerfTest.MapSequenceName=$(MapSequenceComboName) -AutomatedPerfTest.DoCSVProfiler -AutomatedPerfTest.DoInsightsTrace $(InsightsTestCommonArgs)}[Sequence Insights $(MapSequenceComboName)]+" If="$(DoInsightsTrace)"/>
 		</ForEach>
 	</Do>
 
 	<Do If="$(WithAPTStaticCameraTests)">
 		<ForEach Name="StaticCameraMapName" Values="$(StaticCameraMapNames)" Separator="+">
-			<Property Name="TargetTestList" Value="$(TargetTestList)AutomatedPerfTest.StaticCameraTest{-AutomatedPerfTest.StaticCameraPerfTest.MapName=$(StaticCameraMapName) -AutomatedPerfTest.DoCSVProfiler $(AutomatedPerformanceCommonArgs)}[Static Camera Perf $(StaticCameraMapName)]+" If="$(DoPerf)" />
-			<Property Name="TargetTestList" Value="$(TargetTestList)AutomatedPerfTest.StaticCameraTest{-AutomatedPerfTest.StaticCameraPerfTest.MapName=$(StaticCameraMapName) -AutomatedPerfTest.DoCSVProfiler -AutomatedPerfTest.DoLLM=true $(AutomatedPerformanceCommonArgs)}[Static Camera LLM $(StaticCameraMapName)]+" If="$(DoLLM)"/>
-			<Property Name="TargetTestList" Value="$(TargetTestList)AutomatedPerfTest.StaticCameraTest{-AutomatedPerfTest.StaticCameraPerfTest.MapName=$(StaticCameraMapName) -AutomatedPerfTest.DoCSVProfiler -AutomatedPerfTest.DoInsightsTrace=true $(AutomatedPerformanceCommonArgs)}[Static Camera Insights $(StaticCameraMapName)]+" If="$(DoInsightsTrace)"/>
+			<Property Name="TargetTestList" Value="$(TargetTestList)AutomatedPerfTest.StaticCameraTest{-AutomatedPerfTest.StaticCameraPerfTest.MapName=$(StaticCameraMapName) -AutomatedPerfTest.DoCSVProfiler $(PerfTestCommonArgs)}[Static Camera Perf $(StaticCameraMapName)]+" If="$(DoPerf)" />
+			<Property Name="TargetTestList" Value="$(TargetTestList)AutomatedPerfTest.StaticCameraTest{-AutomatedPerfTest.StaticCameraPerfTest.MapName=$(StaticCameraMapName) -AutomatedPerfTest.DoCSVProfiler -AutomatedPerfTest.DoLLM=true $(LLMTestCommonArgs)}[Static Camera LLM $(StaticCameraMapName)]+" If="$(DoLLM)"/>
+			<Property Name="TargetTestList" Value="$(TargetTestList)AutomatedPerfTest.StaticCameraTest{-AutomatedPerfTest.StaticCameraPerfTest.MapName=$(StaticCameraMapName) -AutomatedPerfTest.DoCSVProfiler -AutomatedPerfTest.DoInsightsTrace=true $(InsightsTestCommonArgs)}[Static Camera Insights $(StaticCameraMapName)]+" If="$(DoInsightsTrace)"/>
 		</ForEach>
 	</Do>
 
 	<Do If="$(WithAPTMaterialTest)">
-		<Property Name="TargetTestList" Value="$(TargetTestList)AutomatedPerfTest.MaterialTest{-AutomatedPerfTest.DoCSVProfiler $(AutomatedPerformanceCommonArgs)}[Material Perf]+" If="$(DoPerf)" />
-		<Property Name="TargetTestList" Value="$(TargetTestList)AutomatedPerfTest.MaterialTest{-AutomatedPerfTest.DoCSVProfiler -AutomatedPerfTest.DoLLM $(AutomatedPerformanceCommonArgs)}[Material LLM]+" If="$(DoLLM)"/>
-		<Property Name="TargetTestList" Value="$(TargetTestList)AutomatedPerfTest.MaterialTest{-AutomatedPerfTest.DoCSVProfiler -AutomatedPerfTest.DoInsightsTrace $(AutomatedPerformanceCommonArgs)}[Material Insights]+" If="$(DoInsightsTrace)"/>
+		<Property Name="TargetTestList" Value="$(TargetTestList)AutomatedPerfTest.MaterialTest{-AutomatedPerfTest.DoCSVProfiler $(PerfTestCommonArgs)}[Material Perf]+" If="$(DoPerf)" />
+		<Property Name="TargetTestList" Value="$(TargetTestList)AutomatedPerfTest.MaterialTest{-AutomatedPerfTest.DoCSVProfiler -AutomatedPerfTest.DoLLM $(LLMTestCommonArgs)}[Material LLM]+" If="$(DoLLM)"/>
+		<Property Name="TargetTestList" Value="$(TargetTestList)AutomatedPerfTest.MaterialTest{-AutomatedPerfTest.DoCSVProfiler -AutomatedPerfTest.DoInsightsTrace $(InsightsTestCommonArgs)}[Material Insights]+" If="$(DoInsightsTrace)"/>
 	</Do>
 
 	<Do If="$(WithAPTReplayTest)">
-		<Property Name="TargetTestList" Value="$(TargetTestList)AutomatedPerfTest.ReplayTest{-AutomatedPerfTest.ReplayPerfTest.ReplayName=&quot;$(ReplayName)&quot; -AutomatedPerfTest.DoCSVProfiler $(AutomatedPerformanceCommonArgs)}[Replay Perf]+" If="$(DoPerf)" />
-		<Property Name="TargetTestList" Value="$(TargetTestList)AutomatedPerfTest.ReplayTest{-AutomatedPerfTest.ReplayPerfTest.ReplayName=&quot;$(ReplayName)&quot; -AutomatedPerfTest.DoCSVProfiler -AutomatedPerfTest.DoLLM $(AutomatedPerformanceCommonArgs)}[Replay LLM]+" If="$(DoLLM)"/>
-		<Property Name="TargetTestList" Value="$(TargetTestList)AutomatedPerfTest.ReplayTest{-AutomatedPerfTest.ReplayPerfTest.ReplayName=&quot;$(ReplayName)&quot; -AutomatedPerfTest.DoCSVProfiler -AutomatedPerfTest.DoInsightsTrace $(AutomatedPerformanceCommonArgs)}[Replay Insights]+" If="$(DoInsightsTrace)"/>
+		<Property Name="TargetTestList" Value="$(TargetTestList)AutomatedPerfTest.ReplayTest{-AutomatedPerfTest.ReplayPerfTest.ReplayName=&quot;$(ReplayName)&quot; -AutomatedPerfTest.DoCSVProfiler $(PerfTestCommonArgs)}[Replay Perf]+" If="$(DoPerf)" />
+		<Property Name="TargetTestList" Value="$(TargetTestList)AutomatedPerfTest.ReplayTest{-AutomatedPerfTest.ReplayPerfTest.ReplayName=&quot;$(ReplayName)&quot; -AutomatedPerfTest.DoCSVProfiler -AutomatedPerfTest.DoLLM $(LLMTestCommonArgs)}[Replay LLM]+" If="$(DoLLM)"/>
+		<Property Name="TargetTestList" Value="$(TargetTestList)AutomatedPerfTest.ReplayTest{-AutomatedPerfTest.ReplayPerfTest.ReplayName=&quot;$(ReplayName)&quot; -AutomatedPerfTest.DoCSVProfiler -AutomatedPerfTest.DoInsightsTrace $(InsightsTestCommonArgs)}[Replay Insights]+" If="$(DoInsightsTrace)"/>
 	</Do>
 
 </BuildGraph>
