PerfTool: Adds support for Mac and Linux.

Example commandline to build for Mac:
cd Engine/Source/Programs/CSVTools/
../../../../Engine/Binaries/ThirdParty/DotNet/8.0.300/mac-arm64/dotnet build CSVTools.sln -c Release

Will output to a Platform specific folder (Windows continues to just be at CSVTools/):
CSVStats -> [REDACTED_USER_HOME]/Perforce/zn_ue5_main/Engine/Binaries/DotNET/CsvTools/Mac/CSVStats.dll

Updated the various .cs files to refer to the dlls (on Mac/Linux) and run them via UE's version of dotnet for Mac and Linux.

diff --git a/Engine/Plugins/Performance/AutomatedPerfTesting/Build/Scripts/AutomatedPerfTestNode.cs b/Engine/Plugins/Performance/AutomatedPerfTesting/Build/Scripts/AutomatedPerfTestNode.cs
--- a/Engine/Plugins/Performance/AutomatedPerfTesting/Build/Scripts/AutomatedPerfTestNode.cs
+++ b/Engine/Plugins/Performance/AutomatedPerfTesting/Build/Scripts/AutomatedPerfTestNode.cs
@@ -2,6 +2,7 @@
 
 using System;
 using System.Collections.Generic;
+using System.Runtime.InteropServices;
 using System.IO;
 using System.Linq;
 using EpicGame;
@@ -344,7 +345,21 @@ namespace AutomatedPerfTest
 		{
 			var ReportCacheDir = GetCachedConfiguration().PerfCacheRoot;	// see if this is appropriate
 
-			var ToolPath = FileReference.Combine(Unreal.EngineDirectory, "Binaries", "DotNET", "CsvTools", "PerfreportTool.exe");
+			string perfreportExe = "";
+			if (OperatingSystem.IsWindows())
+			{
+			 	perfreportExe = "PerfreportTool.exe";
+			}
+			else if (OperatingSystem.IsLinux())
+			{
+				perfreportExe = Path.Combine("Linux", "PerfreportTool.dll");
+			}
+			else if (OperatingSystem.IsMacOS())
+			{
+				perfreportExe = Path.Combine("Mac", "PerfreportTool.dll");
+			}
+
+			var ToolPath = FileReference.Combine(Unreal.EngineDirectory, "Binaries", "DotNET", "CsvTools", perfreportExe);
 			if (!FileReference.Exists(ToolPath))
 			{
 				Logger.LogError("Failed to find perf report utility at this path: \"{ToolPath}\".", ToolPath);
@@ -427,7 +442,7 @@ namespace AutomatedPerfTest
 			string SearchPattern = $"{Context.BuildInfo.ProjectName}*";
 			
 			// Produce the detailed report, and update the perf cache
-			CommandUtils.RunAndLog(ToolPath.FullName, $"-csvdir \"{NewestDir}\" -o \"{ReportPath}\" -reportxmlbasedir \"{ReportConfigDir}\" -summaryTable \"{SummaryTable}\" -summaryTableCache \"{ReportCacheDir}\" -searchpattern {SearchPattern} -metadatafilter platform=\"{PlatformNameFilter}\"", out int ErrorCode);
+			CommandUtils.RunAndLog(CommandUtils.CmdEnv, CommandUtils.CmdEnv.DotnetMsbuildPath, $" {ToolPath.FullName} -csvdir \"{NewestDir}\" -o \"{ReportPath}\" -reportxmlbasedir \"{ReportConfigDir}\" -summaryTable \"{SummaryTable}\" -summaryTableCache \"{ReportCacheDir}\" -searchpattern {SearchPattern} -metadatafilter platform=\"{PlatformNameFilter}\"", out int ErrorCode);
 			if (ErrorCode != 0)
 			{
 				Logger.LogError("PerfReportTool returned error code \"{ErrorCode}\" while generating detailed report.", ErrorCode);
@@ -457,6 +472,7 @@ namespace AutomatedPerfTest
 			{
 				var Args = new[]
 				{
+					$"{ToolPath.FullName}",
 					$"-summarytablecachein \"{ReportCacheDir}\"",
 					$"-summaryTableFilename \"{Name}.html\"",
 					$"-reportxmlbasedir \"{ReportConfigDir}\"",
@@ -470,7 +486,7 @@ namespace AutomatedPerfTest
 
 				var ArgStr = string.Join(" ", Args);
 
-				CommandUtils.RunAndLog(ToolPath.FullName, ArgStr, out ErrorCode);
+				CommandUtils.RunAndLog(CommandUtils.CmdEnv, CommandUtils.CmdEnv.DotnetMsbuildPath, ArgStr, out ErrorCode);
 				if (ErrorCode != 0)
 				{
 					Logger.LogError("PerfReportTool returned error code \"{ErrorCode}\" while generating historic report.", ErrorCode);
diff --git a/Engine/Source/Programs/AutomationTool/Gauntlet/Framework/Utils/Gauntlet.ReportGenUtils.cs b/Engine/Source/Programs/AutomationTool/Gauntlet/Framework/Utils/Gauntlet.ReportGenUtils.cs
--- a/Engine/Source/Programs/AutomationTool/Gauntlet/Framework/Utils/Gauntlet.ReportGenUtils.cs
+++ b/Engine/Source/Programs/AutomationTool/Gauntlet/Framework/Utils/Gauntlet.ReportGenUtils.cs
@@ -19,10 +19,8 @@ namespace Gauntlet
 
 		static ReportGenUtils()
 		{
-			char PathDelimter;
 			if (BuildHostPlatform.Current.Platform.IsInGroup("Windows"))
 			{
-				PathDelimter = ';';
 				DotnetPathLocation = "dotnet.exe";
 				BasePythonLocation = @"Engine\Binaries\ThirdParty\Python3\Win64\python.exe";
 			}
@@ -37,17 +35,10 @@ namespace Gauntlet
 				}
 				else
 				{
-					PathDelimter = ':';
-					DotnetPathLocation = "dotnet";
 					BasePythonLocation = @"Engine/Binaries/ThirdParty/Python3/Mac/bin/python3";
 				}
 			}
-
-			DotnetPathLocation = Path.Combine(Environment.GetEnvironmentVariable("PATH")
-				.Split(PathDelimter)
-				.Select(Executable => Executable.TrimEnd('\\').TrimEnd('/'))
-				.Where(Executable => Executable.Contains("dotnet"))
-				.FirstOrDefault(), DotnetPathLocation);
+			DotnetPathLocation = UnrealBuildBase.Unreal.DotnetPath.FullName;
 
 			if (string.IsNullOrEmpty(DotnetPathLocation))
 			{
@@ -87,7 +78,15 @@ namespace Gauntlet
 
 		public static string GetCsvToolPath(string ToolName)
 		{
-			string Extension = BuildHostPlatform.Current.Platform.IsInGroup("Windows") ? ".exe" : ".dll";
+			string Extension = BuildHostPlatform.Current.Platform.IsInGroup(UnrealPlatformGroup.Microsoft) ? ".exe" : ".dll";
+			if (BuildHostPlatform.Current.Platform.IsInGroup(UnrealPlatformGroup.Apple))
+			{
+				return Path.Combine(CommandUtils.CmdEnv.LocalRoot, "Engine", "Binaries", "DotNet", "CsvTools", "Mac",  ToolName) + Extension;
+			}
+			if (BuildHostPlatform.Current.Platform.IsInGroup(UnrealPlatformGroup.Unix))
+			{
+				return Path.Combine(CommandUtils.CmdEnv.LocalRoot, "Engine", "Binaries", "DotNet","CSVTools", "Linux", ToolName) + Extension;
+			}
 			return Path.Combine(CommandUtils.CmdEnv.LocalRoot, "Engine", "Binaries", "DotNet", "CsvTools", ToolName) + Extension;
 		}
 
diff --git a/Engine/Source/Programs/CSVTools/CSVCollate/CSVCollate.csproj b/Engine/Source/Programs/CSVTools/CSVCollate/CSVCollate.csproj
--- a/Engine/Source/Programs/CSVTools/CSVCollate/CSVCollate.csproj
+++ b/Engine/Source/Programs/CSVTools/CSVCollate/CSVCollate.csproj
@@ -1,6 +1,6 @@
 <Project Sdk="Microsoft.NET.Sdk">
   <PropertyGroup>
-    <TargetFramework>net6.0</TargetFramework>
+    <TargetFramework>net8.0</TargetFramework>
     <OutputType>Exe</OutputType>
     <GenerateAssemblyInfo>false</GenerateAssemblyInfo>
     <AppendTargetFrameworkToOutputPath>false</AppendTargetFrameworkToOutputPath>
@@ -10,16 +10,17 @@
 Copyright Epic Games, Inc. All Rights Reserved.</Copyright>
     <AssemblyVersion>1.0.0.0</AssemblyVersion>
     <FileVersion>1.0.0.0</FileVersion>
+    <RuntimeIdentifiers>win-x64;osx-x64;win-arm64;osx-arm64;linux-x64;linux-arm64</RuntimeIdentifiers>
   </PropertyGroup>
   <PropertyGroup Condition=" '$(Configuration)|$(Platform)' == 'Debug|AnyCPU' ">
     <PlatformTarget>x64</PlatformTarget>
   </PropertyGroup>
   <PropertyGroup Condition=" '$(Configuration)|$(Platform)' == 'Release|AnyCPU' ">
     <PlatformTarget>anycpu</PlatformTarget>
-    <OutputPath>..\..\..\..\Binaries\DotNET\CsvTools\</OutputPath>
+     <OutputPath Condition="'$(OutputPath)' == ''">..\..\..\..\Binaries\DotNET\CsvTools\</OutputPath>
   </PropertyGroup>
   <ItemGroup>
-    <ProjectReference Include="..\CsvStats\CSVStats.csproj" />
+    <ProjectReference Include="..\CsvStats\CsvStats.csproj" />
   </ItemGroup>
   <ItemGroup>
     <PackageReference Include="Microsoft.CSharp" Version="4.7.0" />
diff --git a/Engine/Source/Programs/CSVTools/CSVFilter/CSVFilter.csproj b/Engine/Source/Programs/CSVTools/CSVFilter/CSVFilter.csproj
--- a/Engine/Source/Programs/CSVTools/CSVFilter/CSVFilter.csproj
+++ b/Engine/Source/Programs/CSVTools/CSVFilter/CSVFilter.csproj
@@ -1,6 +1,6 @@
 <Project Sdk="Microsoft.NET.Sdk">
   <PropertyGroup>
-    <TargetFramework>net6.0</TargetFramework>
+    <TargetFramework>net8.0</TargetFramework>
     <OutputType>Exe</OutputType>
     <GenerateAssemblyInfo>false</GenerateAssemblyInfo>
     <AppendTargetFrameworkToOutputPath>false</AppendTargetFrameworkToOutputPath>
@@ -10,16 +10,17 @@
 Copyright Epic Games, Inc. All Rights Reserved.</Copyright>
     <AssemblyVersion>1.0.0.0</AssemblyVersion>
     <FileVersion>1.0.0.0</FileVersion>
+    <RuntimeIdentifiers>win-x64;osx-x64;win-arm64;osx-arm64;linux-x64;linux-arm64</RuntimeIdentifiers>
   </PropertyGroup>
   <PropertyGroup Condition=" '$(Configuration)|$(Platform)' == 'Debug|AnyCPU' ">
     <PlatformTarget>x64</PlatformTarget>
   </PropertyGroup>
   <PropertyGroup Condition=" '$(Configuration)|$(Platform)' == 'Release|AnyCPU' ">
     <PlatformTarget>anycpu</PlatformTarget>
-    <OutputPath>..\..\..\..\Binaries\DotNET\CsvTools\</OutputPath>
+     <OutputPath Condition="'$(OutputPath)' == ''">..\..\..\..\Binaries\DotNET\CsvTools\</OutputPath>
   </PropertyGroup>
   <ItemGroup>
-    <ProjectReference Include="..\CsvStats\CSVStats.csproj" />
+    <ProjectReference Include="..\CsvStats\CsvStats.csproj" />
   </ItemGroup>
   <ItemGroup>
     <PackageReference Include="Microsoft.CSharp" Version="4.7.0" />
diff --git a/Engine/Source/Programs/CSVTools/CSVSplit/CSVSplit.csproj b/Engine/Source/Programs/CSVTools/CSVSplit/CSVSplit.csproj
--- a/Engine/Source/Programs/CSVTools/CSVSplit/CSVSplit.csproj
+++ b/Engine/Source/Programs/CSVTools/CSVSplit/CSVSplit.csproj
@@ -1,6 +1,6 @@
 <Project Sdk="Microsoft.NET.Sdk">
   <PropertyGroup>
-    <TargetFramework>net6.0</TargetFramework>
+    <TargetFramework>net8.0</TargetFramework>
     <OutputType>Exe</OutputType>
     <AppendTargetFrameworkToOutputPath>false</AppendTargetFrameworkToOutputPath>
     <GenerateAssemblyInfo>false</GenerateAssemblyInfo>
@@ -10,16 +10,17 @@
 Copyright Epic Games, Inc. All Rights Reserved.</Copyright>
     <AssemblyVersion>1.0.0.0</AssemblyVersion>
     <FileVersion>1.0.0.0</FileVersion>
+    <RuntimeIdentifiers>win-x64;osx-x64;win-arm64;osx-arm64;linux-x64;linux-arm64</RuntimeIdentifiers>
   </PropertyGroup>
   <PropertyGroup Condition=" '$(Configuration)|$(Platform)' == 'Debug|AnyCPU' ">
     <PlatformTarget>x64</PlatformTarget>
   </PropertyGroup>
   <PropertyGroup Condition=" '$(Configuration)|$(Platform)' == 'Release|AnyCPU' ">
     <PlatformTarget>anycpu</PlatformTarget>
-    <OutputPath>..\..\..\..\Binaries\DotNET\CsvTools\</OutputPath>
+     <OutputPath Condition="'$(OutputPath)' == ''">..\..\..\..\Binaries\DotNET\CsvTools\</OutputPath>
   </PropertyGroup>
   <ItemGroup>
-    <ProjectReference Include="..\CsvStats\CSVStats.csproj" />
+    <ProjectReference Include="..\CsvStats\CsvStats.csproj" />
   </ItemGroup>
   <ItemGroup>
     <PackageReference Include="Microsoft.CSharp" Version="4.7.0" />
diff --git a/Engine/Source/Programs/CSVTools/CSVToSVG/CSVToSVG.csproj b/Engine/Source/Programs/CSVTools/CSVToSVG/CSVToSVG.csproj
--- a/Engine/Source/Programs/CSVTools/CSVToSVG/CSVToSVG.csproj
+++ b/Engine/Source/Programs/CSVTools/CSVToSVG/CSVToSVG.csproj
@@ -1,16 +1,17 @@
 <Project Sdk="Microsoft.NET.Sdk">
   <PropertyGroup>
-    <TargetFramework>net6.0</TargetFramework>
+    <TargetFramework>net8.0</TargetFramework>
     <OutputType>Exe</OutputType>
     <GenerateAssemblyInfo>false</GenerateAssemblyInfo>
     <AppendTargetFrameworkToOutputPath>false</AppendTargetFrameworkToOutputPath>
+    <RuntimeIdentifiers>win-x64;osx-x64;win-arm64;osx-arm64;linux-x64;linux-arm64</RuntimeIdentifiers>
   </PropertyGroup>
   <PropertyGroup Condition=" '$(Configuration)|$(Platform)' == 'Debug|AnyCPU' ">
     <PlatformTarget>x64</PlatformTarget>
   </PropertyGroup>
   <PropertyGroup Condition=" '$(Configuration)|$(Platform)' == 'Release|AnyCPU' ">
     <PlatformTarget>anycpu</PlatformTarget>
-    <OutputPath>..\..\..\..\Binaries\DotNET\CsvTools\</OutputPath>
+     <OutputPath Condition="'$(OutputPath)' == ''">..\..\..\..\Binaries\DotNET\CsvTools\</OutputPath>
   </PropertyGroup>
   <PropertyGroup>
     <StartupObject>
@@ -23,7 +24,7 @@ Copyright Epic Games, Inc. All Rights Reserved.</Copyright>
     <FileVersion>1.0.0.0</FileVersion>
   </PropertyGroup>
   <ItemGroup>
-    <ProjectReference Include="..\CsvStats\CSVStats.csproj" />
+    <ProjectReference Include="..\CsvStats\CsvStats.csproj" />
     <ProjectReference Include="..\CsvToSvgLib\CsvToSvgLib.csproj" />
   </ItemGroup>
   <ItemGroup>
diff --git a/Engine/Source/Programs/CSVTools/CsvConvert/CsvConvert.csproj b/Engine/Source/Programs/CSVTools/CsvConvert/CsvConvert.csproj
--- a/Engine/Source/Programs/CSVTools/CsvConvert/CsvConvert.csproj
+++ b/Engine/Source/Programs/CSVTools/CsvConvert/CsvConvert.csproj
@@ -1,6 +1,6 @@
 ﻿<Project Sdk="Microsoft.NET.Sdk">
   <PropertyGroup>
-    <TargetFramework>net6.0</TargetFramework>
+    <TargetFramework>net8.0</TargetFramework>
     <OutputType>Exe</OutputType>
     <GenerateAssemblyInfo>false</GenerateAssemblyInfo>
     <AppendTargetFrameworkToOutputPath>false</AppendTargetFrameworkToOutputPath>
@@ -9,16 +9,17 @@
     <Copyright>Copyright Epic Games, Inc. All rights reserved.</Copyright>
     <AssemblyVersion>1.0.0.0</AssemblyVersion>
     <FileVersion>1.0.0.0</FileVersion>
+    <RuntimeIdentifiers>win-x64;osx-x64;win-arm64;osx-arm64;linux-x64;linux-arm64</RuntimeIdentifiers>
   </PropertyGroup>
   <PropertyGroup Condition=" '$(Configuration)|$(Platform)' == 'Debug|AnyCPU' ">
     <PlatformTarget>x64</PlatformTarget>
   </PropertyGroup>
   <PropertyGroup Condition=" '$(Configuration)|$(Platform)' == 'Release|AnyCPU' ">
     <PlatformTarget>anycpu</PlatformTarget>
-    <OutputPath>..\..\..\..\Binaries\DotNET\CsvTools\</OutputPath>
+     <OutputPath Condition="'$(OutputPath)' == ''">..\..\..\..\Binaries\DotNET\CsvTools\</OutputPath>
   </PropertyGroup>
   <ItemGroup>
-    <ProjectReference Include="..\CsvStats\CSVStats.csproj" />
+    <ProjectReference Include="..\CsvStats\CsvStats.csproj" />
   </ItemGroup>
   <ItemGroup>
     <PackageReference Include="Microsoft.CSharp" Version="4.7.0" />
diff --git a/Engine/Source/Programs/CSVTools/CsvInfo/CsvInfo.csproj b/Engine/Source/Programs/CSVTools/CsvInfo/CsvInfo.csproj
--- a/Engine/Source/Programs/CSVTools/CsvInfo/CsvInfo.csproj
+++ b/Engine/Source/Programs/CSVTools/CsvInfo/CsvInfo.csproj
@@ -1,6 +1,6 @@
 ﻿<Project Sdk="Microsoft.NET.Sdk">
   <PropertyGroup>
-    <TargetFramework>net6.0</TargetFramework>
+    <TargetFramework>net8.0</TargetFramework>
     <OutputType>Exe</OutputType>
     <GenerateAssemblyInfo>false</GenerateAssemblyInfo>
     <AppendTargetFrameworkToOutputPath>false</AppendTargetFrameworkToOutputPath>
@@ -9,16 +9,17 @@
     <Copyright>Copyright Epic Games, Inc. All rights reserved.</Copyright>
     <AssemblyVersion>1.0.0.0</AssemblyVersion>
     <FileVersion>1.0.0.0</FileVersion>
+    <RuntimeIdentifiers>win-x64;osx-x64;win-arm64;osx-arm64;linux-x64;linux-arm64</RuntimeIdentifiers>
   </PropertyGroup>
   <PropertyGroup Condition=" '$(Configuration)|$(Platform)' == 'Debug|AnyCPU' ">
     <PlatformTarget>x64</PlatformTarget>
   </PropertyGroup>
   <PropertyGroup Condition=" '$(Configuration)|$(Platform)' == 'Release|AnyCPU' ">
     <PlatformTarget>anycpu</PlatformTarget>
-    <OutputPath>..\..\..\..\Binaries\DotNET\CsvTools\</OutputPath>
+     <OutputPath Condition="'$(OutputPath)' == ''">..\..\..\..\Binaries\DotNET\CsvTools\</OutputPath>
   </PropertyGroup>
   <ItemGroup>
-    <ProjectReference Include="..\CsvStats\CSVStats.csproj" />
+    <ProjectReference Include="..\CsvStats\CsvStats.csproj" />
   </ItemGroup>
   <ItemGroup>
     <PackageReference Include="Microsoft.CSharp" Version="4.7.0" />
diff --git a/Engine/Source/Programs/CSVTools/CsvStats/CsvStats.csproj b/Engine/Source/Programs/CSVTools/CsvStats/CsvStats.csproj
--- a/Engine/Source/Programs/CSVTools/CsvStats/CsvStats.csproj
+++ b/Engine/Source/Programs/CSVTools/CsvStats/CsvStats.csproj
@@ -1,19 +1,20 @@
 <Project Sdk="Microsoft.NET.Sdk">
   <PropertyGroup>
     <AppendTargetFrameworkToOutputPath>false</AppendTargetFrameworkToOutputPath>
-    <TargetFramework>net6.0</TargetFramework>
+    <TargetFramework>net8.0</TargetFramework>
     <OutputType>Library</OutputType>
     <Copyright>Copyright %28C%29 Microsoft. All rights reserved.
 Copyright Epic Games, Inc. All Rights Reserved.</Copyright>
     <AssemblyVersion>1.0.0.0</AssemblyVersion>
     <FileVersion>1.0.0.0</FileVersion>
     <GenerateAssemblyInfo>false</GenerateAssemblyInfo>
+    <RuntimeIdentifiers>win-x64;osx-x64;win-arm64;osx-arm64;linux-x64;linux-arm64</RuntimeIdentifiers>
   </PropertyGroup>
   <PropertyGroup Condition=" '$(Configuration)|$(Platform)' == 'Debug|AnyCPU' ">
     <PlatformTarget>x64</PlatformTarget>
   </PropertyGroup>
   <PropertyGroup Condition=" '$(Configuration)|$(Platform)' == 'Release|AnyCPU' ">
-    <OutputPath>..\..\..\..\Binaries\DotNET\CsvTools\</OutputPath>
+    <OutputPath Condition="'$(OutputPath)' == ''">..\..\..\..\Binaries\DotNET\CsvTools\</OutputPath>
     <PlatformTarget>anycpu</PlatformTarget>
   </PropertyGroup>
   <ItemGroup>
diff --git a/Engine/Source/Programs/CSVTools/CsvToSvgLib/CsvToSvgLib.csproj b/Engine/Source/Programs/CSVTools/CsvToSvgLib/CsvToSvgLib.csproj
--- a/Engine/Source/Programs/CSVTools/CsvToSvgLib/CsvToSvgLib.csproj
+++ b/Engine/Source/Programs/CSVTools/CsvToSvgLib/CsvToSvgLib.csproj
@@ -1,6 +1,6 @@
 ﻿<Project Sdk="Microsoft.NET.Sdk">
   <PropertyGroup>
-    <TargetFramework>net6.0</TargetFramework>
+    <TargetFramework>net8.0</TargetFramework>
     <OutputType>Library</OutputType>
     <RootNamespace>CsvStats</RootNamespace>
     <GenerateAssemblyInfo>false</GenerateAssemblyInfo>
@@ -11,16 +11,17 @@
 Copyright Epic Games, Inc. All Rights Reserved.</Copyright>
     <AssemblyVersion>1.0.0.0</AssemblyVersion>
     <FileVersion>1.0.0.0</FileVersion>
+    <RuntimeIdentifiers>win-x64;osx-x64;win-arm64;osx-arm64;linux-x64;linux-arm64</RuntimeIdentifiers>
   </PropertyGroup>
   <PropertyGroup Condition=" '$(Configuration)|$(Platform)' == 'Debug|AnyCPU' ">
     <PlatformTarget>x64</PlatformTarget>
   </PropertyGroup>
   <PropertyGroup Condition=" '$(Configuration)|$(Platform)' == 'Release|AnyCPU' ">
-    <OutputPath>..\..\..\..\Binaries\DotNET\CsvTools\</OutputPath>
+     <OutputPath Condition="'$(OutputPath)' == ''">..\..\..\..\Binaries\DotNET\CsvTools\</OutputPath>
     <PlatformTarget>anycpu</PlatformTarget>
   </PropertyGroup>
   <ItemGroup>
-    <ProjectReference Include="..\CsvStats\CSVStats.csproj" />
+    <ProjectReference Include="..\CsvStats\CsvStats.csproj" />
   </ItemGroup>
   <ItemGroup>
     <PackageReference Include="Microsoft.CSharp" Version="4.7.0" />
diff --git a/Engine/Source/Programs/CSVTools/PerfReportTool/PerfReportTool.cs b/Engine/Source/Programs/CSVTools/PerfReportTool/PerfReportTool.cs
--- a/Engine/Source/Programs/CSVTools/PerfReportTool/PerfReportTool.cs
+++ b/Engine/Source/Programs/CSVTools/PerfReportTool/PerfReportTool.cs
@@ -3,6 +3,7 @@
 
 using System;
 using System.Collections.Generic;
+using System.Runtime.InteropServices;
 using System.Linq;
 using System.Text;
 using System.IO;
@@ -16,6 +17,7 @@ using System.Threading;
 using PerfSummaries;
 using CSVTools;
 
+
 using System.Text.Json;
 
 namespace PerfReportTool
@@ -2191,7 +2193,20 @@ namespace PerfReportTool
 
 		Process LaunchCsvToSvgAsync(string args)
 		{
-			string csvToolPath = Path.Combine(GetBaseDirectory(), "CSVToSVG.exe");
+			string csvToolPath = "";
+			if (OperatingSystem.IsWindows())
+			{
+			 	csvToolPath = Path.Combine(GetBaseDirectory(), "CSVToSVG.exe");
+			}
+			else if (OperatingSystem.IsLinux())
+			{
+				csvToolPath = Path.Combine("Linux", "CSVToSVG.dll");
+			}
+			else if (OperatingSystem.IsMacOS())
+			{
+				csvToolPath = Path.Combine("Mac", "CSVToSVG.dll");
+			}
+
 			string binary = csvToolPath;
 
 			if (Host == HostPlatform.Windows)
@@ -2201,11 +2216,11 @@ namespace PerfReportTool
 					binary = binary.Replace("\\perfreporttool\\", "\\CsvToSVG\\", StringComparison.InvariantCultureIgnoreCase);
 				}
 			}
-			// run mono on non-Windows hosts
-			else
+			else if (OperatingSystem.IsLinux() || OperatingSystem.IsMacOS())
 			{
-				// note, on Mac mono will not be on path
-				binary = Host == HostPlatform.Linux ? "mono" : "/Library/Frameworks/Mono.framework/Versions/Current/Commands/mono";
+				// Get the dotnet that we were running with and just spawn
+				// using the same one.
+				binary = Process.GetCurrentProcess().MainModule.FileName;
 				args = csvToolPath + " " + args;
 			}
 
diff --git a/Engine/Source/Programs/CSVTools/PerfReportTool/PerfreportTool.csproj b/Engine/Source/Programs/CSVTools/PerfReportTool/PerfreportTool.csproj
--- a/Engine/Source/Programs/CSVTools/PerfReportTool/PerfreportTool.csproj
+++ b/Engine/Source/Programs/CSVTools/PerfReportTool/PerfreportTool.csproj
@@ -1,22 +1,23 @@
 <Project Sdk="Microsoft.NET.Sdk">
   <PropertyGroup>
-    <TargetFramework>net6.0</TargetFramework>
+    <TargetFramework>net8.0</TargetFramework>
     <OutputType>Exe</OutputType>
     <GenerateAssemblyInfo>false</GenerateAssemblyInfo>
     <AppendTargetFrameworkToOutputPath>false</AppendTargetFrameworkToOutputPath>
+    <RuntimeIdentifiers>win-x64;osx-x64;win-arm64;osx-arm64;linux-x64;linux-arm64</RuntimeIdentifiers>
   </PropertyGroup>
   <PropertyGroup Condition=" '$(Configuration)|$(Platform)' == 'Debug|AnyCPU' ">
     <PlatformTarget>x64</PlatformTarget>
   </PropertyGroup>
   <PropertyGroup Condition=" '$(Configuration)|$(Platform)' == 'Release|AnyCPU' ">
     <PlatformTarget>anycpu</PlatformTarget>
-    <OutputPath>..\..\..\..\Binaries\DotNET\CsvTools\</OutputPath>
+     <OutputPath Condition="'$(OutputPath)' == ''">..\..\..\..\Binaries\DotNET\CsvTools\</OutputPath>
   </PropertyGroup>
   <ItemGroup>
     <None Remove="Properties\AssemblyInfo.cs~RF273305d.TMP" />
   </ItemGroup>
   <ItemGroup>
-    <ProjectReference Include="..\CsvStats\CSVStats.csproj" />
+    <ProjectReference Include="..\CsvStats\CsvStats.csproj" />
     <ProjectReference Include="..\CsvToSvgLib\CsvToSvgLib.csproj" />
   </ItemGroup>
   <ItemGroup>
diff --git a/Engine/Source/Programs/CSVTools/PerfReportTool/ReportXml.cs b/Engine/Source/Programs/CSVTools/PerfReportTool/ReportXml.cs
--- a/Engine/Source/Programs/CSVTools/PerfReportTool/ReportXml.cs
+++ b/Engine/Source/Programs/CSVTools/PerfReportTool/ReportXml.cs
@@ -246,7 +246,7 @@ namespace PerfReportTool
 				graphsXMLFilename = reportTypesElement.GetSafeAttribute<string>("reportGraphsFile");
 				if (graphsXMLFilename != null)
 				{
-					graphsXMLFilename = Path.GetDirectoryName(reportTypeXmlFilename) + "\\" + graphsXMLFilename;
+					graphsXMLFilename = Path.Combine(Path.GetDirectoryName(reportTypeXmlFilename), graphsXMLFilename);
 				}
 				else
 				{
