Fix Lyra APT ReplayRun for Consoles

This change enables APT Replay Runs to run on consoles. Previously, the target device (except Win64) would not be able to find the replay file as it was never copied to the device nor could it read it from the host itself without relevant changes to the replay path.

* Copy replay files to relevant consoles depending on platform.
* If device supports host mounting, then update the replay path according to platform.
* Ensure APT Replay Run looks for replay file in the default Replay/Demos folder before failing the test.

diff --git a/Engine/Plugins/Performance/AutomatedPerfTesting/Build/Scripts/AutomatedPerfTestNode.cs b/Engine/Plugins/Performance/AutomatedPerfTesting/Build/Scripts/AutomatedPerfTestNode.cs
--- a/Engine/Plugins/Performance/AutomatedPerfTesting/Build/Scripts/AutomatedPerfTestNode.cs
+++ b/Engine/Plugins/Performance/AutomatedPerfTesting/Build/Scripts/AutomatedPerfTestNode.cs
@@ -763,28 +763,96 @@ namespace AutomatedPerfTest
 	{
 		public AutomatedReplayPerfTestNode(UnrealTestContext InContext) : base(InContext)
 		{
+			Config = null;
+		}
+
+		/// <summary>
+		/// Handler to explicitly copy files to the device. It is assumed this 
+		/// is called when Gauntlet configures the device before starting the
+		/// tests. 
+		/// </summary>
+		public void CopyReplayToDevice(ITargetDevice Device)
+		{
+			if (Config == null)
+			{
+				Log.Warning("ConfigureDevice() called before Test Node Configuration is set.");
+				return;
+			}
+
+			// At the moment in some pathways we don't seem to copy
+			// files in a role. This ensures we have a copy of the 
+			// replay files on the target device regardless. 
+			if (Config.GetRequiredRoles(UnrealTargetRole.Client).Any())
+			{
+				foreach (UnrealTestRole ClientRole in Config.GetRequiredRoles(UnrealTargetRole.Client))
+				{
+					UnrealTargetPlatform Platform = Context.GetRoleContext(UnrealTargetRole.Client).Platform;
+					if (Device.Platform == Platform && ClientRole.FilesToCopy.Any())
+					{
+						Device.CopyAdditionalFiles(ClientRole.FilesToCopy);
+					}
+				}
+			}
 		}
 
 		public override TConfigClass GetConfiguration()
 		{
-			TConfigClass Config = base.GetConfiguration();
+			Config = base.GetConfiguration();
+			
+			Config.DataSourceName = Config.GetDataSourceName(Context.BuildInfo.ProjectName, ReplayDataSourceName);
 
-			Config.DataSourceName = Config.GetDataSourceName(Context.BuildInfo.ProjectName, "ReplayRun");
+			// Since we do not have a device instance at this stage. These platform support
+			// instances would help if we would still like to know if the target platform has
+			// certain features or functionality we need to be aware of in order to pass the
+			// right configuration to Gauntlet.
+			IEnumerable<IPlatformTargetSupport> PlatformSupportInstances = Gauntlet.Utils.InterfaceHelpers.FindImplementations<IPlatformTargetSupport>();
 
+			IEnumerable<UnrealTestRole> ClientRoles = Config.GetRequiredRoles(UnrealTargetRole.Client);
 			// extend the role(s) that we initialized in the base class
-			if (Config.GetRequiredRoles(UnrealTargetRole.Client).Any())
+			if (ClientRoles.Any())
 			{
-				foreach (UnrealTestRole ClientRole in Config.GetRequiredRoles(UnrealTargetRole.Client))
+				foreach (UnrealTestRole ClientRole in ClientRoles)
 				{
 					ClientRole.Controllers.Add("AutomatedReplayPerfTest");
-					ClientRole.CommandLineParams.AddUnique($"AutomatedPerfTest.ReplayPerfTest.ReplayName",
-								Config.ReplayName);
-					Log.Info($"AutomatedPerfTest.ReplayPerfTest.ReplayName={Config.ReplayName}");
+					ClientRole.FilesToCopy.Clear();
+
+					UnrealTargetPlatform Platform = Context.GetRoleContext(UnrealTargetRole.Client).Platform;
+					IPlatformTargetSupport MountablePlatform = PlatformSupportInstances
+						.Where(TargetSupport => TargetSupport.IsHostMountingSupported() && TargetSupport.Platform == Platform)
+						.FirstOrDefault();
+
+					string ReplayName = Path.GetFullPath(Config.ReplayName);
+					if (MountablePlatform != null) 
+					{
+						// If current device supports host mounting of files, update the path and use host mounting 
+						ReplayName = MountablePlatform.GetHostMountedPath(ReplayName);
+						Log.Info($"Host Mountable Platform Detected. Updated Replay File Path to: {ReplayName} ");
+					}
+					else
+					{
+						// Copy replay file to Demos folder in device if host mounting is not supported
+						Log.Info($"Copy Replay File Path to Device: {ReplayName}");
+						UnrealFileToCopy FileToCopy = new UnrealFileToCopy(ReplayName, EIntendedBaseCopyDirectory.Demos, Path.GetFileName(ReplayName));
+						ClientRole.FilesToCopy.Add(FileToCopy);
+						ClientRole.ConfigureDevice = CopyReplayToDevice;
+
+						// If we copy to the "Demos" folder, we just need to pass the Replay Name without
+						// path and extension as the replay subsystem will automatically pick up the file
+						// from this folder.
+						ReplayName = Path.GetFileNameWithoutExtension(ReplayName);
+					}
+
+					ClientRole.CommandLineParams.AddUnique(ReplayParamName, ReplayName);
+					Log.Info($"{ReplayParamName}=\"{ReplayName}\"");
 				}
 			}
 
 			return Config;
 		}
+
+		private readonly string ReplayDataSourceName = "ReplayRun";
+		private readonly string ReplayParamName = "AutomatedPerfTest.ReplayPerfTest.ReplayName";
+		private TConfigClass Config;
 	}
 
 	/// <summary>
diff --git a/Engine/Plugins/Performance/AutomatedPerfTesting/Source/AutomatedPerfTesting/Private/AutomatedReplayPerfTest.cpp b/Engine/Plugins/Performance/AutomatedPerfTesting/Source/AutomatedPerfTesting/Private/AutomatedReplayPerfTest.cpp
--- a/Engine/Plugins/Performance/AutomatedPerfTesting/Source/AutomatedPerfTesting/Private/AutomatedReplayPerfTest.cpp
+++ b/Engine/Plugins/Performance/AutomatedPerfTesting/Source/AutomatedPerfTesting/Private/AutomatedReplayPerfTest.cpp
@@ -91,6 +91,10 @@ void UAutomatedReplayPerfTest::OnInit()
 
 	UE_LOG(LogAutomatedPerfTest, Log, TEXT("AutomatedReplayPerfTest::OnInit"));
 
+	// Replay subsystem looks for replay files in this folder by default.
+	const FString DefaultReplayPath = FPaths::Combine(*FPaths::ProjectSavedDir(), TEXT("Demos/")); 
+	const FString DefaultReplayExtension = TEXT(".replay");
+
 	const UAutomatedReplayPerfTestProjectSettings* Settings = GetDefault<UAutomatedReplayPerfTestProjectSettings>();
 	SetCSVOutputMode(Settings->CSVOutputMode);
 
@@ -99,6 +103,17 @@ void UAutomatedReplayPerfTest::OnInit()
 	ReplayName = CmdReplayName.IsEmpty() ? Settings->ReplayName : CmdReplayName; // Command Line takes precedence
 
 	bool bReplayExists = FPaths::FileExists(ReplayName);
+	if (!bReplayExists)
+	{
+		// On some devices, we may have to copy the replay file to the device
+		// as it may not support reading from the host directly. In that case
+		// we assume only the replay file name is supplied without the whole 
+		// path. We check the default replay path in this scenario. 
+		const FString ReplayAlternatePath = DefaultReplayPath + ReplayName + DefaultReplayExtension;
+		bReplayExists = FPaths::FileExists(ReplayAlternatePath);
+		UE_LOG(LogAutomatedPerfTest, Log, TEXT("Replay found in default Demos path: %s"), *ReplayAlternatePath);
+	}
+
 	if (ReplayName.IsEmpty() || !bReplayExists)
 	{
 		UE_LOG(LogAutomatedPerfTest, Error, TEXT("Replay not found:%s"), *ReplayName);
diff --git a/Engine/Source/Programs/AutomationTool/Gauntlet/Framework/Base/Gauntlet.TargetDevice.cs b/Engine/Source/Programs/AutomationTool/Gauntlet/Framework/Base/Gauntlet.TargetDevice.cs
--- a/Engine/Source/Programs/AutomationTool/Gauntlet/Framework/Base/Gauntlet.TargetDevice.cs
+++ b/Engine/Source/Programs/AutomationTool/Gauntlet/Framework/Base/Gauntlet.TargetDevice.cs
@@ -223,4 +223,43 @@ namespace Gauntlet
 
 		public virtual bool NeedBuildDeployed() => true;
 	}
+
+	/// <summary>
+	/// If platform supports host mounting of files or other platform specific
+	/// functionality, this interface provides a way to provide that information
+	/// without needing an instance of the device.
+	/// </summary>
+	public interface IPlatformTargetSupport
+	{
+		/// <summary>
+		/// If host mounting is supported, return the equivalent
+		/// host mounted path for given platform.
+		/// </summary>
+		/// <param name="Path">Input path</param>
+		/// <returns></returns>
+		string GetHostMountedPath(string Path);
+
+		/// <summary>
+		/// Returns true if host mounting is supported on given
+		/// platform. 
+		/// </summary>
+		/// <returns></returns>
+		bool IsHostMountingSupported();
+
+		/// <summary>
+		/// The target platform
+		/// </summary>
+		UnrealTargetPlatform? Platform { get; }
+	}
+
+	/// <summary>
+	/// Default abstract implementation of IPlatformTargetSupport
+	/// so that target platforms can override only parts they need
+	/// </summary>
+	public abstract class TargetPlatformSupportBase : IPlatformTargetSupport
+	{
+		public virtual string GetHostMountedPath(string Path) => Path;
+		public virtual bool IsHostMountingSupported() => false;
+		public virtual UnrealTargetPlatform? Platform => throw new NotImplementedException();
+	}
 }
\ No newline at end of file
diff --git a/Engine/Source/Programs/AutomationTool/Gauntlet/Platform/Android/Gauntlet.TargetDeviceAndroid.cs b/Engine/Source/Programs/AutomationTool/Gauntlet/Platform/Android/Gauntlet.TargetDeviceAndroid.cs
--- a/Engine/Source/Programs/AutomationTool/Gauntlet/Platform/Android/Gauntlet.TargetDeviceAndroid.cs
+++ b/Engine/Source/Programs/AutomationTool/Gauntlet/Platform/Android/Gauntlet.TargetDeviceAndroid.cs
@@ -2017,4 +2017,10 @@ namespace Gauntlet
 			}
 		}
 	}
+
+	class AndroidPlatformSupport : TargetPlatformSupportBase
+	{
+		public override UnrealTargetPlatform? Platform => UnrealTargetPlatform.Android;
+		public override bool IsHostMountingSupported() => false;
+	}
 }
\ No newline at end of file
diff --git a/Engine/Source/Programs/AutomationTool/Gauntlet/Platform/IOS/Gauntlet.TargetDeviceIOS.cs b/Engine/Source/Programs/AutomationTool/Gauntlet/Platform/IOS/Gauntlet.TargetDeviceIOS.cs
--- a/Engine/Source/Programs/AutomationTool/Gauntlet/Platform/IOS/Gauntlet.TargetDeviceIOS.cs
+++ b/Engine/Source/Programs/AutomationTool/Gauntlet/Platform/IOS/Gauntlet.TargetDeviceIOS.cs
@@ -1478,4 +1478,10 @@ namespace Gauntlet
 		protected override BuildFlags SupportedBuildTypes => BuildFlags.Packaged | BuildFlags.CanReplaceCommandLine | BuildFlags.CanReplaceExecutable | BuildFlags.Bulk | BuildFlags.NotBulk;
 		protected override UnrealTargetPlatform? Platform => UnrealTargetPlatform.IOS;
 	}
+
+	class IOSPlatformSupport : TargetPlatformSupportBase
+	{
+		public override UnrealTargetPlatform? Platform => UnrealTargetPlatform.IOS;
+		public override bool IsHostMountingSupported() => false;
+	}
 }
\ No newline at end of file
diff --git a/Engine/Source/Programs/AutomationTool/Gauntlet/Platform/Windows/Gauntlet.TargetDeviceWindows.cs b/Engine/Source/Programs/AutomationTool/Gauntlet/Platform/Windows/Gauntlet.TargetDeviceWindows.cs
--- a/Engine/Source/Programs/AutomationTool/Gauntlet/Platform/Windows/Gauntlet.TargetDeviceWindows.cs
+++ b/Engine/Source/Programs/AutomationTool/Gauntlet/Platform/Windows/Gauntlet.TargetDeviceWindows.cs
@@ -458,4 +458,11 @@ namespace Gauntlet
 			return new TargetDeviceWindows(InRef, InCachePath);
 		}
 	}
+
+	public class WindowsPlatformSupport : TargetPlatformSupportBase
+	{
+		public override UnrealTargetPlatform? Platform => UnrealTargetPlatform.Win64;
+		public override bool IsHostMountingSupported() => true;
+		public override string GetHostMountedPath(string Path) => Path;
+	}
 }
\ No newline at end of file
