Fix CitySample APT LLM/Insights

* Add option `UsePlatformExtraDebugMemory` to APT common settings. This will allow us to specify if a test needs extra debug memory on the target platform.
* Added support to specify the target device(PS5 only at the moment) to use additional direct memory if above option is true.

diff --git a/Engine/Plugins/Performance/AutomatedPerfTesting/Build/Inc/AutomatedPerfTestCommonSettings.xml b/Engine/Plugins/Performance/AutomatedPerfTesting/Build/Inc/AutomatedPerfTestCommonSettings.xml
--- a/Engine/Plugins/Performance/AutomatedPerfTesting/Build/Inc/AutomatedPerfTestCommonSettings.xml
+++ b/Engine/Plugins/Performance/AutomatedPerfTesting/Build/Inc/AutomatedPerfTestCommonSettings.xml
@@ -27,6 +27,7 @@
 	<Option Name="IterationsPerf" DefaultValue="6" Restrict="^[0-9]+$" Description="Number of iterations to run perf test with" />
 	<Option Name="IterationsLLM" DefaultValue="3" Restrict="^[0-9]+$" Description="Number of iterations to run LLM test with" />
 	<Option Name="IterationsInsights" DefaultValue="1" Restrict="^[0-9]+$" Description="Number of iterations to run Insights test with" />
+	<Option Name="UsePlatformExtraDebugMemory" DefaultValue="true" Restrict="true|false" Description="Use additional debug memory on target device if supported. Applies only to LLM and Insights trace."/>
 
 	<Option Name="ExtraAutomatedPerformanceCommonArgs" DefaultValue="" Description="Additional Common arguments passed to Gauntlet command line"/>
 
@@ -60,6 +61,11 @@
 	<Property Name="LLMTestCommonArgs"      Value="$(AutomatedPerformanceCommonArgs) -iterations=$(IterationsLLM)"/>
 	<Property Name="InsightsTestCommonArgs" Value="$(AutomatedPerformanceCommonArgs) -iterations=$(IterationsInsights)"/>
 
+	<Do If="$(UsePlatformExtraDebugMemory)">
+		<Property Name="LLMTestCommonArgs"      Value="$(LLMTestCommonArgs) -AutomatedPerfTest.UsePlatformExtraDebugMemory"/>
+		<Property Name="InsightsTestCommonArgs" Value="$(InsightsTestCommonArgs) -AutomatedPerfTest.UsePlatformExtraDebugMemory"/>
+	</Do>
+
 	<!-- append plugin-provided tests to GauntletSettings' TargetTestList -->
 	<Do If="$(WithAPTSequenceTests)">
 		<ForEach Name="MapSequenceComboName" Values="$(MapSequenceComboNames)" Separator="+">
diff --git a/Engine/Plugins/Performance/AutomatedPerfTesting/Build/Scripts/AutomatedPerfTestConfig.cs b/Engine/Plugins/Performance/AutomatedPerfTesting/Build/Scripts/AutomatedPerfTestConfig.cs
--- a/Engine/Plugins/Performance/AutomatedPerfTesting/Build/Scripts/AutomatedPerfTestConfig.cs
+++ b/Engine/Plugins/Performance/AutomatedPerfTesting/Build/Scripts/AutomatedPerfTestConfig.cs
@@ -107,6 +107,14 @@ namespace AutomatedPerfTest
 		/// </summary>
 		[AutoParamWithNames("",  "AutomatedPerfTest.ReportPath")]
 		public string ReportPath;
+
+		/// <summary>
+		/// Specify whether this test requires additional debug memory
+		/// on target device. The allocation of this memory is entirely
+		/// dependent on the target device and may not be guaranteed.
+		/// </summary>
+		[AutoParamWithNames(false, "AutomatedPerfTest.UsePlatformExtraDebugMemory")]
+		public bool UsePlatformExtraDebugMemory;
 		
 		public string DataSourceName;
 
@@ -136,6 +144,15 @@ namespace AutomatedPerfTest
 			
 			return $"Automation.{ProjectName}.{dataSourceType}";
 		}
+
+		public override void ApplyToConfig(UnrealAppConfig AppConfig, UnrealSessionRole ConfigRole, IEnumerable<UnrealSessionRole> OtherRoles)
+		{
+			base.ApplyToConfig(AppConfig, ConfigRole, OtherRoles);
+
+			// Let the app know if we need to use extra debug memory before we 
+			// instantiate it. 
+			AppConfig.UsePlatformExtraDebugMemory = UsePlatformExtraDebugMemory;
+		}
 	}
 
 	public class AutomatedSequencePerfTestConfig : AutomatedPerfTestConfigBase
diff --git a/Engine/Source/Programs/AutomationTool/Gauntlet/Unreal/Base/Gauntlet.UnrealAppConfig.cs b/Engine/Source/Programs/AutomationTool/Gauntlet/Unreal/Base/Gauntlet.UnrealAppConfig.cs
--- a/Engine/Source/Programs/AutomationTool/Gauntlet/Unreal/Base/Gauntlet.UnrealAppConfig.cs
+++ b/Engine/Source/Programs/AutomationTool/Gauntlet/Unreal/Base/Gauntlet.UnrealAppConfig.cs
@@ -114,6 +114,10 @@ namespace Gauntlet
 		// Performs a full clean on the device before installing
 		public bool FullClean => DevicePool.FullClean;
 
+		/// <summary>
+		/// Use additional debug memory on target device if available. 
+		/// </summary>
+		public bool UsePlatformExtraDebugMemory { get; set; }
 
 		/// <summary>
 		/// Constructor that sets some required values to defaults
