LaneStrate APT Integration

* Integrate APT framework into Lanestrate. By default, the sequence `Flythrough` will be used to run APT Sequence tests
* Allow APT Tests to succeed even if logging is disabled in Test builds if `APTIgnoreTestBuildLogging` is true

diff --git a/Engine/Plugins/Performance/AutomatedPerfTesting/Build/Inc/AutomatedPerfTestCommonSettings.xml b/Engine/Plugins/Performance/AutomatedPerfTesting/Build/Inc/AutomatedPerfTestCommonSettings.xml
--- a/Engine/Plugins/Performance/AutomatedPerfTesting/Build/Inc/AutomatedPerfTestCommonSettings.xml
+++ b/Engine/Plugins/Performance/AutomatedPerfTesting/Build/Inc/AutomatedPerfTestCommonSettings.xml
@@ -28,6 +28,7 @@
 	<Option Name="IterationsLLM" DefaultValue="3" Restrict="^[0-9]+$" Description="Number of iterations to run LLM test with" />
 	<Option Name="IterationsInsights" DefaultValue="1" Restrict="^[0-9]+$" Description="Number of iterations to run Insights test with" />
 	<Option Name="UsePlatformExtraDebugMemory" DefaultValue="true" Restrict="true|false" Description="Use additional debug memory on target device if supported. Applies only to LLM and Insights trace."/>
+	<Option Name="APTIgnoreTestBuildLogging" Restrict="true|false" DefaultValue="true" Description="If true, tests run in Test build configuration will not fail if logs are empty." />
 
 	<Option Name="ExtraAutomatedPerformanceCommonArgs" DefaultValue="" Description="Additional Common arguments passed to Gauntlet command line"/>
 
@@ -55,6 +56,7 @@
 	<Property Name="AutomatedPerformanceCommonArgs" Value="$(AutomatedPerformanceCommonArgs) -NoLocalReports" If="!$(GenerateLocalReports) And $(DoCSVProfiler)" />
 	<Property Name="AutomatedPerformanceCommonArgs" Value="$(AutomatedPerformanceCommonArgs) -resX=$(ResX) -resY=$(ResY)"/>
 	<Property Name="AutomatedPerformanceCommonArgs" Value="$(AutomatedPerformanceCommonArgs) -AutomatedPerfTest.DoVideoCapture" If="$(DoVideoCapture)"/>
+	<Property Name="AutomatedPerformanceCommonArgs" Value="$(AutomatedPerformanceCommonArgs) -AutomatedPerfTest.IgnoreTestBuildLogging" If="$(APTIgnoreTestBuildLogging)"/>
 	<Property Name="AutomatedPerformanceCommonArgs" Value="$(AutomatedPerformanceCommonArgs) $(ExtraAutomatedPerformanceCommonArgs)" If="'$(ExtraAutomatedPerformanceCommonArgs)' != ''"/>
 
 	<Property Name="PerfTestCommonArgs"     Value="$(AutomatedPerformanceCommonArgs) -iterations=$(IterationsPerf)"/>
diff --git a/Engine/Plugins/Performance/AutomatedPerfTesting/Build/Scripts/AutomatedPerfTestConfig.cs b/Engine/Plugins/Performance/AutomatedPerfTesting/Build/Scripts/AutomatedPerfTestConfig.cs
--- a/Engine/Plugins/Performance/AutomatedPerfTesting/Build/Scripts/AutomatedPerfTestConfig.cs
+++ b/Engine/Plugins/Performance/AutomatedPerfTesting/Build/Scripts/AutomatedPerfTestConfig.cs
@@ -116,6 +116,15 @@ namespace AutomatedPerfTest
 		[AutoParamWithNames(false, "AutomatedPerfTest.UsePlatformExtraDebugMemory")]
 		public bool UsePlatformExtraDebugMemory;
 		
+		/// <summary>
+		/// If true, logs will not be considered when passing or failing
+		/// a test in Test builds. Useful if logging is disabled in Test
+		/// builds and we do not want to fail the test if no logs are 
+		/// found. 
+		/// </summary>
+		[AutoParamWithNames(false, "AutomatedPerfTest.IgnoreTestBuildLogging")]
+		public bool IgnoreTestBuildLogging;
+
 		public string DataSourceName;
 
 		/// <summary>
diff --git a/Engine/Plugins/Performance/AutomatedPerfTesting/Build/Scripts/AutomatedPerfTestNode.cs b/Engine/Plugins/Performance/AutomatedPerfTesting/Build/Scripts/AutomatedPerfTestNode.cs
--- a/Engine/Plugins/Performance/AutomatedPerfTesting/Build/Scripts/AutomatedPerfTestNode.cs
+++ b/Engine/Plugins/Performance/AutomatedPerfTesting/Build/Scripts/AutomatedPerfTestNode.cs
@@ -15,6 +15,7 @@ using UnrealBuildBase;	// for Unreal.RootDirectory
 using UnrealBuildTool;	// for UnrealTargetPlatform
 
 using static AutomationTool.CommandUtils;
+using System.Runtime.CompilerServices;
 
 namespace AutomatedPerfTest
 {
@@ -209,6 +210,7 @@ namespace AutomatedPerfTest
 		{
 			// Check for login failure
 			UnrealLogParser Parser = new UnrealLogParser(InArtifacts.AppInstance.GetLogReader());
+			TConfigClass Config = GetCachedConfiguration();
 
 			ExitReason = "";
 			ExitCode = -1;
@@ -224,6 +226,45 @@ namespace AutomatedPerfTest
 				}
 			}
 
+			// If this is a Test Target Configuration and we have configured to ignore logging,
+			// we can check the process exit code and return from here. This is especially useful
+			// when logging is disabled in Test builds, but we still want to exit the test
+			// successfully. 
+			UnrealTargetConfiguration TargetConfig = InArtifacts.SessionRole.Configuration;
+			bool bIgnoreLoggingInTest = Config.IgnoreTestBuildLogging && TargetConfig == UnrealTargetConfiguration.Test;
+			bool bTestHasErrorLog = InLogSummary.FatalError != null || InLogSummary.HasTestExitCode;
+
+			// This is a major assumption that Gauntlet captures the process exit code on all
+			// platforms and is the only thing indicating if the test has actually passed or
+			// failed in the absence of logging. 
+			int ProcessExitCode = InArtifacts.AppInstance.ExitCode;
+			if (InReason == StopReason.Completed && bIgnoreLoggingInTest && ProcessExitCode == 0)
+			{
+				ExitCode = 0;
+				ExitReason = "Test build exited successfully without logs.";
+				return UnrealProcessResult.ExitOk;
+			}
+			else if (InReason == StopReason.Completed && bIgnoreLoggingInTest && ProcessExitCode != 0 && !bTestHasErrorLog)
+			{
+				// Process has not exited cleanly and we do not have any error log messages. 
+				// Fail the test here and provide context to user.
+				ExitCode = ProcessExitCode;
+				ExitReason = "Test build exited with error. Please enable logging in build for more information.";
+				return UnrealProcessResult.TestFailure;
+			}
+
+			// Let the user know that if their tests are failing but process exit code is 0
+			// and they have logging disabled in Test builds that they can ignore logging
+			// failures if they wish to do so. 
+			if(!InLogSummary.EngineInitialized && TargetConfig == UnrealTargetConfiguration.Test && ProcessExitCode == 0)
+			{
+				Logger.LogWarning("*** Engine Initialization log not detected in Test build with Exit Code 0. " +
+					"This test will fail. " +
+					"Try passing `-AutomatedPerfTest.IgnoreTestBuildLogging` while running this test " +
+					"or pass `-set:APTIgnoreTestBuildLogging=true` if running via BuildGraph if you " +
+					"wish to ignore log parsing checks in this test or recompile with logging enabled. ***");
+			}
+
 			return base.GetExitCodeAndReason(InReason, InLogSummary, InArtifacts, out ExitReason, out ExitCode);
 		}
 
