APT project: fix missing EngineDir and double inclusion.

- CS files located in the same dir as .csproj do not need to be referenced by it, dotnet will pull it in manually (thanks JoeK!).
- cleaned up some other leftovers.

diff --git a/Engine/Plugins/Performance/AutomatedPerfTesting/Build/Scripts/AutomatedPerfTesting.Automation.csproj b/Engine/Plugins/Performance/AutomatedPerfTesting/Build/Scripts/AutomatedPerfTesting.Automation.csproj
--- a/Engine/Plugins/Performance/AutomatedPerfTesting/Build/Scripts/AutomatedPerfTesting.Automation.csproj
+++ b/Engine/Plugins/Performance/AutomatedPerfTesting/Build/Scripts/AutomatedPerfTesting.Automation.csproj
@@ -1,5 +1,8 @@
 <Project Sdk="Microsoft.NET.Sdk">
-  <Import Project="AsisNovaMaps.AutoPerfTest.csproj.props" Condition="Exists('AsisNovaMaps.AutoPerfTest.csproj.props')" />
+  <Import Project="AutomatedPerfTesting.Automation.csproj.props" Condition="Exists('AutomatedPerfTesting.Automation.csproj.props')" />
+  <PropertyGroup>
+    <EngineDir>../../../../..</EngineDir>
+  </PropertyGroup>
   <Import Project="$(EngineDir)\Source\Programs\Shared\UnrealEngine.csproj.props" />
   <PropertyGroup>
     <TargetFramework>net8.0</TargetFramework>
@@ -7,17 +10,13 @@
     <AppendTargetFrameworkToOutputPath>false</AppendTargetFrameworkToOutputPath>
     <GenerateTargetFrameworkAttribute>false</GenerateTargetFrameworkAttribute>
     <Configurations>Debug;Release;Development</Configurations>
-    <RootNamespace>AsisNovaMaps.AutoPerfTest</RootNamespace>
-    <AssemblyName>AsisNovaMaps.AutoPerfTest</AssemblyName>
-    <OutputPath>..\..\Binaries\DotNET\AutomationTool\AutomationScripts</OutputPath>
+    <RootNamespace>AutomatedPerfTesting.Automation</RootNamespace>
+    <AssemblyName>AutomatedPerfTesting.Automation</AssemblyName>
+    <OutputPath>$(EngineDir)\Binaries\DotNET\AutomationTool\AutomationScripts</OutputPath>
     <TreatWarningsAsErrors>true</TreatWarningsAsErrors>
     <WarningsNotAsErrors>612,618</WarningsNotAsErrors>
     <DebugType>pdbonly</DebugType>
   </PropertyGroup>
-  <ItemGroup>
-    <Compile Include="$(EngineDir)\Plugins\Performance\AutomatedPerfTesting\Build\Scripts\AutomatedPerfTestConfig.cs" />
-    <Compile Include="$(EngineDir)\Plugins\Performance\AutomatedPerfTesting\Build\Scripts\AutomatedPerfTestNode.cs" />
-  </ItemGroup>
   <ItemGroup>
     <ProjectReference Include="$(EngineDir)\Source\Programs\AutomationTool\AutomationUtils\AutomationUtils.Automation.csproj" />
     <ProjectReference Include="$(EngineDir)\Source\Programs\Shared\EpicGames.Core\EpicGames.Core.csproj" PrivateAssets="All">
