Fix running Windows PerftestTool to use the PerftestTool.dll rather than the exe

#rnx

diff --git a/Engine/Plugins/Performance/AutomatedPerfTesting/Build/Scripts/AutomatedPerfTestNode.cs b/Engine/Plugins/Performance/AutomatedPerfTesting/Build/Scripts/AutomatedPerfTestNode.cs
--- a/Engine/Plugins/Performance/AutomatedPerfTesting/Build/Scripts/AutomatedPerfTestNode.cs
+++ b/Engine/Plugins/Performance/AutomatedPerfTesting/Build/Scripts/AutomatedPerfTestNode.cs
@@ -386,17 +386,10 @@ namespace AutomatedPerfTest
 		{
 			var ReportCacheDir = GetCachedConfiguration().PerfCacheRoot;	// see if this is appropriate
 
-			string perfreportExe = "";
-			if (OperatingSystem.IsWindows())
-			{
-			 	perfreportExe = "PerfreportTool.exe";
-			}
-			else if (OperatingSystem.IsLinux() || OperatingSystem.IsMacOS())
-			{
-				perfreportExe = Path.Combine("PerfreportTool.dll");
-			}
-
-			var ToolPath = FileReference.Combine(Unreal.EngineDirectory, "Binaries", "DotNET", "CsvTools", perfreportExe);
+			string perfreportTool = "";
+			perfreportTool = Path.Combine("PerfreportTool.dll");
+	
+			var ToolPath = FileReference.Combine(Unreal.EngineDirectory, "Binaries", "DotNET", "CsvTools", perfreportTool);
 			if (!FileReference.Exists(ToolPath))
 			{
 				Logger.LogError("Failed to find perf report utility at this path: \"{ToolPath}\".", ToolPath);
