Fix OutputPath for several automation scripts that write to the same directory

#rnx

diff --git a/Engine/Plugins/Performance/AutomatedPerfTesting/Build/Scripts/AutomatedPerfTesting.Automation.csproj b/Engine/Plugins/Performance/AutomatedPerfTesting/Build/Scripts/AutomatedPerfTesting.Automation.csproj
--- a/Engine/Plugins/Performance/AutomatedPerfTesting/Build/Scripts/AutomatedPerfTesting.Automation.csproj
+++ b/Engine/Plugins/Performance/AutomatedPerfTesting/Build/Scripts/AutomatedPerfTesting.Automation.csproj
@@ -12,7 +12,7 @@
     <Configurations>Debug;Release;Development</Configurations>
     <RootNamespace>AutomatedPerfTesting.Automation</RootNamespace>
     <AssemblyName>AutomatedPerfTesting.Automation</AssemblyName>
-    <OutputPath>$(EngineDir)\Binaries\DotNET\AutomationTool\AutomationScripts</OutputPath>
+    <OutputPath>$(EngineDir)\Binaries\DotNET\AutomationTool\AutomationScripts\AutomatedPerfTesting</OutputPath>
     <TreatWarningsAsErrors>true</TreatWarningsAsErrors>
     <WarningsNotAsErrors>612,618</WarningsNotAsErrors>
     <DebugType>pdbonly</DebugType>
diff --git a/Engine/Plugins/VirtualProduction/ICVFXTesting/Build/Scripts/ICVFX.Automation.csproj b/Engine/Plugins/VirtualProduction/ICVFXTesting/Build/Scripts/ICVFX.Automation.csproj
--- a/Engine/Plugins/VirtualProduction/ICVFXTesting/Build/Scripts/ICVFX.Automation.csproj
+++ b/Engine/Plugins/VirtualProduction/ICVFXTesting/Build/Scripts/ICVFX.Automation.csproj
@@ -13,7 +13,7 @@
     <RootNamespace>ICVFX.Automation</RootNamespace>
     <AssemblyName>ICVFX.Automation</AssemblyName>
     <WarningsNotAsErrors>612,618</WarningsNotAsErrors>
-    <OutputPath>..\..\..\..\..\Binaries\DotNET\AutomationTool\AutomationScripts\</OutputPath>
+    <OutputPath>..\..\..\..\..\Binaries\DotNET\AutomationTool\AutomationScripts\ICVFX</OutputPath>
     <TreatWarningsAsErrors>true</TreatWarningsAsErrors>
     <DebugType>pdbonly</DebugType>
   </PropertyGroup>
