APT: Add GPUPerf test type.

- Reduces GPU async work in the same way as ReplayRun.py does
- Locks dynamic res (where supported) to a configurable value.
- Minor clean ups and fixes in the code.

diff --git a/Engine/Plugins/Performance/AutomatedPerfTesting/Build/Inc/AutomatedPerfTestCommonSettings.xml b/Engine/Plugins/Performance/AutomatedPerfTesting/Build/Inc/AutomatedPerfTestCommonSettings.xml
--- a/Engine/Plugins/Performance/AutomatedPerfTesting/Build/Inc/AutomatedPerfTestCommonSettings.xml
+++ b/Engine/Plugins/Performance/AutomatedPerfTesting/Build/Inc/AutomatedPerfTestCommonSettings.xml
@@ -26,6 +26,7 @@
 	<Option Name="ResY" DefaultValue="2160" Restrict="^[0-9]+$" Description="Resolution height to run tests with" />
 	<Option Name="IterationsPerf" DefaultValue="6" Restrict="^[0-9]+$" Description="Number of iterations to run perf test with" />
 	<Option Name="IterationsLLM" DefaultValue="3" Restrict="^[0-9]+$" Description="Number of iterations to run LLM test with" />
+	<Option Name="IterationsGPUPerf" DefaultValue="3" Restrict="^[0-9]+$" Description="Number of iterations to run GPUPerf test with" />
 	<Option Name="IterationsInsights" DefaultValue="1" Restrict="^[0-9]+$" Description="Number of iterations to run Insights test with" />
 	<Option Name="UsePlatformExtraDebugMemory" DefaultValue="true" Restrict="true|false" Description="Use additional debug memory on target device if supported. Applies only to LLM and Insights trace."/>
 	<Option Name="APTIgnoreTestBuildLogging" Restrict="true|false" DefaultValue="true" Description="If true, tests run in Test build configuration will not fail if logs are empty." />
@@ -41,6 +42,7 @@
 	<!-- Test sub-types, act as flavors to the above tests -->
 	<Option Name="DoPerf" DefaultValue="false" Restrict="true|false" Description="Whether to run Perf (general performance) test." />
 	<Option Name="DoLLM" DefaultValue="false" Restrict="true|false" Description="Whether to run LLM (memory) test." />
+	<Option Name="DoGPUPerf" DefaultValue="false" Restrict="true|false" Description="Whether to run GPUPerf (GPU performance) test." />
 	<Option Name="DoInsightsTrace" DefaultValue="false" Restrict="true|false" Description="Whether or not to capture Unreal Insights trace during the run." />
 
 	<!-- Construct the common arguments for the performance test and report commands -->
@@ -61,6 +63,7 @@
 
 	<Property Name="PerfTestCommonArgs"     Value="$(AutomatedPerformanceCommonArgs) -iterations=$(IterationsPerf)"/>
 	<Property Name="LLMTestCommonArgs"      Value="$(AutomatedPerformanceCommonArgs) -iterations=$(IterationsLLM)"/>
+	<Property Name="GPUPerfTestCommonArgs"      Value="$(AutomatedPerformanceCommonArgs) -iterations=$(IterationsGPUPerf)"/>
 	<Property Name="InsightsTestCommonArgs" Value="$(AutomatedPerformanceCommonArgs) -iterations=$(IterationsInsights)"/>
 
 	<Do If="$(UsePlatformExtraDebugMemory)">
@@ -73,6 +76,7 @@
 		<ForEach Name="MapSequenceComboName" Values="$(MapSequenceComboNames)" Separator="+">
 			<Property Name="TargetTestList" Value="$(TargetTestList)AutomatedPerfTest.SequenceTest{-AutomatedPerfTest.SequencePerfTest.MapSequenceName=$(MapSequenceComboName) -AutomatedPerfTest.DoCSVProfiler $(PerfTestCommonArgs)}[Sequence Perf $(MapSequenceComboName)]+" If="$(DoPerf)" />
 			<Property Name="TargetTestList" Value="$(TargetTestList)AutomatedPerfTest.SequenceTest{-AutomatedPerfTest.SequencePerfTest.MapSequenceName=$(MapSequenceComboName) -AutomatedPerfTest.DoCSVProfiler -AutomatedPerfTest.DoLLM $(LLMTestCommonArgs)}[Sequence LLM $(MapSequenceComboName)]+" If="$(DoLLM)"/>
+			<Property Name="TargetTestList" Value="$(TargetTestList)AutomatedPerfTest.SequenceTest{-AutomatedPerfTest.SequencePerfTest.MapSequenceName=$(MapSequenceComboName) -AutomatedPerfTest.DoCSVProfiler -AutomatedPerfTest.DoGPUPerf $(GPUPerfTestCommonArgs)}[Sequence GPUPerf $(MapSequenceComboName)]+" If="$(DoGPUPerf)" />
 			<Property Name="TargetTestList" Value="$(TargetTestList)AutomatedPerfTest.SequenceTest{-AutomatedPerfTest.SequencePerfTest.MapSequenceName=$(MapSequenceComboName) -AutomatedPerfTest.DoCSVProfiler -AutomatedPerfTest.DoInsightsTrace $(InsightsTestCommonArgs)}[Sequence Insights $(MapSequenceComboName)]+" If="$(DoInsightsTrace)"/>
 		</ForEach>
 	</Do>
@@ -80,20 +84,23 @@
 	<Do If="$(WithAPTStaticCameraTests)">
 		<ForEach Name="StaticCameraMapName" Values="$(StaticCameraMapNames)" Separator="+">
 			<Property Name="TargetTestList" Value="$(TargetTestList)AutomatedPerfTest.StaticCameraTest{-AutomatedPerfTest.StaticCameraPerfTest.MapName=$(StaticCameraMapName) -AutomatedPerfTest.DoCSVProfiler $(PerfTestCommonArgs)}[Static Camera Perf $(StaticCameraMapName)]+" If="$(DoPerf)" />
-			<Property Name="TargetTestList" Value="$(TargetTestList)AutomatedPerfTest.StaticCameraTest{-AutomatedPerfTest.StaticCameraPerfTest.MapName=$(StaticCameraMapName) -AutomatedPerfTest.DoCSVProfiler -AutomatedPerfTest.DoLLM=true $(LLMTestCommonArgs)}[Static Camera LLM $(StaticCameraMapName)]+" If="$(DoLLM)"/>
-			<Property Name="TargetTestList" Value="$(TargetTestList)AutomatedPerfTest.StaticCameraTest{-AutomatedPerfTest.StaticCameraPerfTest.MapName=$(StaticCameraMapName) -AutomatedPerfTest.DoCSVProfiler -AutomatedPerfTest.DoInsightsTrace=true $(InsightsTestCommonArgs)}[Static Camera Insights $(StaticCameraMapName)]+" If="$(DoInsightsTrace)"/>
+			<Property Name="TargetTestList" Value="$(TargetTestList)AutomatedPerfTest.StaticCameraTest{-AutomatedPerfTest.StaticCameraPerfTest.MapName=$(StaticCameraMapName) -AutomatedPerfTest.DoCSVProfiler -AutomatedPerfTest.DoLLM $(LLMTestCommonArgs)}[Static Camera LLM $(StaticCameraMapName)]+" If="$(DoLLM)"/>
+			<Property Name="TargetTestList" Value="$(TargetTestList)AutomatedPerfTest.StaticCameraTest{-AutomatedPerfTest.StaticCameraPerfTest.MapName=$(StaticCameraMapName) -AutomatedPerfTest.DoCSVProfiler -AutomatedPerfTest.DoGPUPerf $(LLMTestCommonArgs)}[Static Camera GPUPerf $(StaticCameraMapName)]+" If="$(DoGPUPerf)"/>
+			<Property Name="TargetTestList" Value="$(TargetTestList)AutomatedPerfTest.StaticCameraTest{-AutomatedPerfTest.StaticCameraPerfTest.MapName=$(StaticCameraMapName) -AutomatedPerfTest.DoCSVProfiler -AutomatedPerfTest.DoInsightsTrace $(InsightsTestCommonArgs)}[Static Camera Insights $(StaticCameraMapName)]+" If="$(DoInsightsTrace)"/>
 		</ForEach>
 	</Do>
 
 	<Do If="$(WithAPTMaterialTest)">
 		<Property Name="TargetTestList" Value="$(TargetTestList)AutomatedPerfTest.MaterialTest{-AutomatedPerfTest.DoCSVProfiler $(PerfTestCommonArgs)}[Material Perf]+" If="$(DoPerf)" />
 		<Property Name="TargetTestList" Value="$(TargetTestList)AutomatedPerfTest.MaterialTest{-AutomatedPerfTest.DoCSVProfiler -AutomatedPerfTest.DoLLM $(LLMTestCommonArgs)}[Material LLM]+" If="$(DoLLM)"/>
+		<Property Name="TargetTestList" Value="$(TargetTestList)AutomatedPerfTest.MaterialTest{-AutomatedPerfTest.DoCSVProfiler -AutomatedPerfTest.DoGPUPerf $(LLMTestCommonArgs)}[Material GPUPerf]+" If="$(DoGPUPerf)"/>
 		<Property Name="TargetTestList" Value="$(TargetTestList)AutomatedPerfTest.MaterialTest{-AutomatedPerfTest.DoCSVProfiler -AutomatedPerfTest.DoInsightsTrace $(InsightsTestCommonArgs)}[Material Insights]+" If="$(DoInsightsTrace)"/>
 	</Do>
 
 	<Do If="$(WithAPTReplayTest)">
 		<Property Name="TargetTestList" Value="$(TargetTestList)AutomatedPerfTest.ReplayTest{-AutomatedPerfTest.ReplayPerfTest.ReplayName=&quot;$(ReplayName)&quot; -AutomatedPerfTest.DoCSVProfiler $(PerfTestCommonArgs)}[Replay Perf]+" If="$(DoPerf)" />
 		<Property Name="TargetTestList" Value="$(TargetTestList)AutomatedPerfTest.ReplayTest{-AutomatedPerfTest.ReplayPerfTest.ReplayName=&quot;$(ReplayName)&quot; -AutomatedPerfTest.DoCSVProfiler -AutomatedPerfTest.DoLLM $(LLMTestCommonArgs)}[Replay LLM]+" If="$(DoLLM)"/>
+		<Property Name="TargetTestList" Value="$(TargetTestList)AutomatedPerfTest.ReplayTest{-AutomatedPerfTest.ReplayPerfTest.ReplayName=&quot;$(ReplayName)&quot; -AutomatedPerfTest.DoCSVProfiler -AutomatedPerfTest.DoGPUPerf $(LLMTestCommonArgs)}[Replay GPUPerf]+" If="$(DoGPUPerf)"/>
 		<Property Name="TargetTestList" Value="$(TargetTestList)AutomatedPerfTest.ReplayTest{-AutomatedPerfTest.ReplayPerfTest.ReplayName=&quot;$(ReplayName)&quot; -AutomatedPerfTest.DoCSVProfiler -AutomatedPerfTest.DoInsightsTrace $(InsightsTestCommonArgs)}[Replay Insights]+" If="$(DoInsightsTrace)"/>
 	</Do>
 
diff --git a/Engine/Plugins/Performance/AutomatedPerfTesting/Build/Scripts/AutomatedPerfTestConfig.cs b/Engine/Plugins/Performance/AutomatedPerfTesting/Build/Scripts/AutomatedPerfTestConfig.cs
--- a/Engine/Plugins/Performance/AutomatedPerfTesting/Build/Scripts/AutomatedPerfTestConfig.cs
+++ b/Engine/Plugins/Performance/AutomatedPerfTesting/Build/Scripts/AutomatedPerfTestConfig.cs
@@ -70,6 +70,18 @@ namespace AutomatedPerfTest
 		[AutoParamWithNames(false, "AutomatedPerfTest.DoLLM")]
 		public bool DoLLM;
 
+		/// <summary>
+		/// Whether to run GPU perf test, that is, dynres is fixed, while GPU async work gets reduced, serializing passes to be timed separately
+		/// </summary>
+		[AutoParamWithNames(false, "AutomatedPerfTest.DoGPUPerf")]
+		public bool DoGPUPerf;
+
+		/// <summary>
+		/// Whether to lock dynamic resolution. This is applied automatically when GPUPerf test is run.
+		/// </summary>
+		[AutoParamWithNames(false,"AutomatedPerfTest.LockDynamicRes")]
+		public bool LockDynamicRes;
+
 		/// <summary>
 		/// Let BuildGraph tell us where we should output the Insights trace after running a test so that we know where to grab it from when we're done
 		/// </summary>
diff --git a/Engine/Plugins/Performance/AutomatedPerfTesting/Build/Scripts/AutomatedPerfTestNode.cs b/Engine/Plugins/Performance/AutomatedPerfTesting/Build/Scripts/AutomatedPerfTestNode.cs
--- a/Engine/Plugins/Performance/AutomatedPerfTesting/Build/Scripts/AutomatedPerfTestNode.cs
+++ b/Engine/Plugins/Performance/AutomatedPerfTesting/Build/Scripts/AutomatedPerfTestNode.cs
@@ -283,13 +283,16 @@ namespace AutomatedPerfTest
 					// to preserve them until we're ready to make our report on the final iteration.
 					CopyPerfFilesToTempDir(ArtifactPath);
 
+					bool bGeneratedLocalReport = false;
+					bool bGeneratedPRSReport = false;
+
 					// Local report generation is useful for people conducting tests locally without a centralized server (for A/B testing for instance).
 					// To make it work out of the box, enable by default for non-build machine runs
 					if (Globals.Params.ParseParam("LocalReports") || (!CommandUtils.IsBuildMachine && !Globals.Params.ParseParam("NoLocalReports")))
 					{
 						// NOTE: This does not currently work with long paths due to the CsvTools not properly supporting them.
 						Log.Info("Generating performance reports using PerfReportTool.");
-						GenerateLocalPerfReport(Context.GetRoleContext(UnrealTargetRole.Client).Platform, ArtifactPath);
+						bGeneratedLocalReport = GenerateLocalPerfReport(Context.GetRoleContext(UnrealTargetRole.Client).Platform, ArtifactPath);
 					}
 
 					// On build machines, default to producing PRS report
@@ -304,6 +307,8 @@ namespace AutomatedPerfTest
 						Log.Info("Creating perf server importer with build name {BuildName}", OriginalBuildName);
 						string DataSourceName = GetConfiguration().DataSourceName;
 						string ImportDirOverride = Globals.Params.ParseValue("PerfReportServerImportDir", null);
+						Logger.LogInformation("Creating PRS Importer for data source '{0}' and import dir override (if any) '{1}'.",
+							DataSourceName, ImportDirOverride);
 						ICsvImporter Importer = ReportGenUtils.CreatePerfReportServerImporter(DataSourceName, OriginalBuildName,
 							CommandUtils.IsBuildMachine, ImportDirOverride, CommonDataSourceFields);
 						if (Importer != null)
@@ -363,11 +368,22 @@ namespace AutomatedPerfTest
 
 							// Create the import batch
 							Importer.Import(CsvImportEntries);
+							// trust blindly for now...
+							bGeneratedPRSReport = true;
+						}
+						else
+						{
+							Logger.LogWarning("Unable to create PRS Importer.");
 						}
 
 						// Cleanup the temp dir
 						TempPerfCSVDir.Delete(recursive: true);
 					}
+
+					if (!bGeneratedLocalReport && !bGeneratedPRSReport)
+					{
+						Logger.LogWarning("Did not generate neither local nor a PRS report.");
+					}
 				}
 			}
 			else
@@ -382,7 +398,7 @@ namespace AutomatedPerfTest
 		/// Produces a detailed csv report using PerfReportTool.
 		/// Also, stores perf data in the perf cache, and generates a historic report using the data the cache contains.
 		/// </summary>
-		private void GenerateLocalPerfReport(UnrealTargetPlatform Platform, string ArtifactPath)
+		private bool GenerateLocalPerfReport(UnrealTargetPlatform Platform, string ArtifactPath)
 		{
 			var ReportCacheDir = GetCachedConfiguration().PerfCacheRoot;	// see if this is appropriate
 
@@ -393,7 +409,7 @@ namespace AutomatedPerfTest
 			if (!FileReference.Exists(ToolPath))
 			{
 				Logger.LogError("Failed to find perf report utility at this path: \"{ToolPath}\".", ToolPath);
-				return;
+				return false;
 			}
 
 			var ReportConfigDir = GetCachedConfiguration().ReportConfigDir;
@@ -430,7 +446,7 @@ namespace AutomatedPerfTest
 			if (DiscoveredCsvs.Count == 0)
 			{
 				Logger.LogError("Test completed successfully but no csv profiling results were found. Searched paths were:\r\n  {Paths}", string.Join("\r\n  ", CsvsPaths.Select(s => $"\"{s}\"")));
-				return;
+				return false;
 			}
 
 			// Find the newest csv file and get its directory
@@ -453,7 +469,7 @@ namespace AutomatedPerfTest
 				catch (Exception Ex)
 				{
 					Logger.LogError("Failed to create perf report cache directory \"{ReportCacheDir}\". {Ex}", ReportCacheDir, Ex);
-					return;
+					return false;
 				}
 			}
 			if (!Directory.Exists(ReportPath))
@@ -462,7 +478,7 @@ namespace AutomatedPerfTest
 				catch (Exception Ex)
 				{
 					Logger.LogError("Failed to create perf report output directory \"{ReportPath}\". {Ex}", ReportPath, Ex);
-					return;
+					return false;
 				}
 			}
 
@@ -522,6 +538,8 @@ namespace AutomatedPerfTest
 					Logger.LogError("PerfReportTool returned error code \"{ErrorCode}\" while generating historic report.", ErrorCode);
 				}
 			}
+
+			return true;
 		}
 
 		private void CopyPerfFilesToTempDir(string ArtifactPath)
@@ -602,6 +620,11 @@ namespace AutomatedPerfTest
 				return "LLM";
 			}
 
+			if (Config.DoGPUPerf)
+			{
+				return "GPUPerf";
+			}
+
 			return "Perf";
 		}
 
@@ -648,7 +671,30 @@ namespace AutomatedPerfTest
 				ClientRole.CommandLineParams.Add("llmcsv");
 			}
 
-			if (Config.DoFPSChart)
+			if (Config.DoGPUPerf)
+			{
+				// see ReplayRun.py, reducedAsyncComputeCommands
+				// We enable r.nanite.asyncrasterization.shadowdepths because nanite overlaps with itself, so it's possible to time accurately without distorting other timings
+				ClientRole.CommandLineParams.AddOrAppendParamValue("execcmds", "r.StencilLODMode 1");
+				ClientRole.CommandLineParams.AddOrAppendParamValue("execcmds", "r.VolumetricRenderTarget.PreferAsyncCompute 0");
+				ClientRole.CommandLineParams.AddOrAppendParamValue("execcmds", "r.LumenScene.Lighting.AsyncCompute 0");
+				ClientRole.CommandLineParams.AddOrAppendParamValue("execcmds", "r.Lumen.DiffuseIndirect.AsyncCompute 0");
+				ClientRole.CommandLineParams.AddOrAppendParamValue("execcmds", "r.Bloom.AsyncCompute 0");
+				ClientRole.CommandLineParams.AddOrAppendParamValue("execcmds", "r.nanite.asyncrasterization.shadowdepths 1");
+				ClientRole.CommandLineParams.AddOrAppendParamValue("execcmds", "r.TSR.AsyncCompute 0");
+				ClientRole.CommandLineParams.AddOrAppendParamValue("execcmds", "r.RayTracing.AsyncBuild 0");
+				ClientRole.CommandLineParams.AddOrAppendParamValue("execcmds", "r.DFShadowAsyncCompute 0");
+				ClientRole.CommandLineParams.AddOrAppendParamValue("execcmds", "r.AmbientOcclusion.Compute 1"); // 1 here means compute on the graphics pipe
+				ClientRole.CommandLineParams.AddOrAppendParamValue("execcmds", "r.LocalFogVolume.TileCullingUseAsync 0");
+				ClientRole.CommandLineParams.AddOrAppendParamValue("execcmds", "r.SkyAtmosphereASyncCompute 0");
+			}
+
+			if (Config.DoGPUPerf || Config.LockDynamicRes)
+			{
+				ClientRole.CommandLineParams.Add("AutomatedPerfTest.LockDynamicRes");
+			}
+
+			if(Config.DoFPSChart)
 			{
 				ClientRole.CommandLineParams.Add("AutomatedPerfTest.DoFPSChart");
 			}
@@ -672,6 +718,11 @@ namespace AutomatedPerfTest
 					CsvMetadata.Add("changelist=" + Context.BuildInfo.Changelist);
 				}
 
+				if (Config.DoGPUPerf)
+				{
+					CsvMetadata.Add("ReducedAsyncCompute=1");
+				}
+
 				ClientRole.CommandLineParams.Add("csvMetadata", "\"" + String.Join(",", CsvMetadata) + "\"");
 			}
 
diff --git a/Engine/Plugins/Performance/AutomatedPerfTesting/Config/Engine.ini b/Engine/Plugins/Performance/AutomatedPerfTesting/Config/Engine.ini
new file mode 100644
--- /dev/null
+++ b/Engine/Plugins/Performance/AutomatedPerfTesting/Config/Engine.ini
@@ -0,0 +1,3 @@
+[SystemSettings]
+; if you need to override this per platform, this needs to be done in Engine/Platforms/<Platform>/Plugins/Performance/AutomatedPerfTesting/Config/<Platform>Engine.ini
+APT.DynamicRes.LockedScreenPercentage=75
diff --git a/Engine/Plugins/Performance/AutomatedPerfTesting/Resources/TemplatesForHookup/RunLocalTests.bat b/Engine/Plugins/Performance/AutomatedPerfTesting/Resources/TemplatesForHookup/RunLocalTests.bat
--- a/Engine/Plugins/Performance/AutomatedPerfTesting/Resources/TemplatesForHookup/RunLocalTests.bat
+++ b/Engine/Plugins/Performance/AutomatedPerfTesting/Resources/TemplatesForHookup/RunLocalTests.bat
@@ -1,3 +1,3 @@
-call **RUN_UAT_RELATIVEPATH**\RunUAT.bat BuildGraph -Script=**REPLACE_PROJECTPATH**/Build/AutoPerfTests.xml -Target="BuildAndTest **REPLACE_PROJECTNAME**" -UseLocalBuildStorage -set:TargetPlatforms=Win64 -set:TargetConfigurations=Test -set:WithAPTSequenceTests=true -set:DoPerf=true -set:DoLLM=true -set:GenerateLocalReports=true -set:GenerateLocalReports=true -set:IterationsPerf=1 -set:IterationsLLM=1 -set:IterationsInsights=1
+call **RUN_UAT_RELATIVEPATH**\RunUAT.bat BuildGraph -Script=**REPLACE_PROJECTPATH**/Build/AutoPerfTests.xml -Target="BuildAndTest **REPLACE_PROJECTNAME**" -UseLocalBuildStorage -set:TargetPlatforms=Win64 -set:TargetConfigurations=Test -set:TargetBootTest=false -set:WithAPTSequenceTests=true -set:DoPerf=true -set:DoLLM=true -set:GenerateLocalReports=true -set:IterationsPerf=1 -set:IterationsLLM=1 -set:IterationsInsights=1
 pause
 
diff --git a/Engine/Plugins/Performance/AutomatedPerfTesting/Source/AutomatedPerfTesting/Private/AutomatedPerfTestControllerBase.cpp b/Engine/Plugins/Performance/AutomatedPerfTesting/Source/AutomatedPerfTesting/Private/AutomatedPerfTestControllerBase.cpp
--- a/Engine/Plugins/Performance/AutomatedPerfTesting/Source/AutomatedPerfTesting/Private/AutomatedPerfTestControllerBase.cpp
+++ b/Engine/Plugins/Performance/AutomatedPerfTesting/Source/AutomatedPerfTesting/Private/AutomatedPerfTestControllerBase.cpp
@@ -24,11 +24,31 @@
 DEFINE_LOG_CATEGORY(LogAutomatedPerfTest)
 CSV_DEFINE_CATEGORY(AutomatedPerfTest, true);
 
+static float GAPTDynamicResLockedScreenPercentage = 100.f;
+static FAutoConsoleVariableRef CVarDynamicResLockedScreenPercentage(
+	TEXT("APT.DynamicRes.LockedScreenPercentage"),
+	GAPTDynamicResLockedScreenPercentage,
+	TEXT("Target resolution percentage, configurable per platform. Use -AutomatedPerfTest.LockDynamicRes to force the resolution to this"),
+	ECVF_Default
+);
 
 void UAutomatedPerfTestControllerBase::OnPreWorldInitializeInternal(UWorld* World, const UWorld::InitializationValues IVS)
 {
 	TryEarlyExec(World);
 	OnPreWorldInitialize(World);
+
+	if (RequestsLockedDynRes())
+	{
+		IConsoleVariable* CVarTestScreenPercentage = IConsoleManager::Get().FindConsoleVariable(TEXT("r.dynamicres.testscreenpercentage"));
+		if (ensure(CVarTestScreenPercentage))
+		{
+			if (CVarTestScreenPercentage->GetFloat() != GAPTDynamicResLockedScreenPercentage)
+			{
+				UE_LOG(LogAutomatedPerfTest, Display, TEXT("Locking screen percentage to %.2f"), GAPTDynamicResLockedScreenPercentage);
+				CVarTestScreenPercentage->Set(GAPTDynamicResLockedScreenPercentage);
+			}
+		}
+	}
 }
 
 void UAutomatedPerfTestControllerBase::OnPreWorldInitialize(UWorld* World)
@@ -94,6 +114,7 @@ UAutomatedPerfTestControllerBase::UAutomatedPerfTestControllerBase(const FObject
 	, bRequestsInsightsTrace(false)
 	, bRequestsCSVProfiler(false)
 	, bRequestsVideoCapture(false)
+	, bRequestsLockedDynRes(false)
     , CSVOutputMode(EAutomatedPerfTestCSVOutputMode::Single)
 {
 	// cache this off once, so that it's consistent throughout a session
@@ -155,6 +176,11 @@ bool UAutomatedPerfTestControllerBase::RequestsVideoCapture() const
 	return bRequestsVideoCapture;
 }
 
+bool UAutomatedPerfTestControllerBase::RequestsLockedDynRes() const
+{
+	return bRequestsLockedDynRes;
+}
+
 bool UAutomatedPerfTestControllerBase::TryStartInsightsTrace()
 {
 	const FString TraceFileName = GetTestID() + ".utrace";
@@ -497,26 +523,31 @@ void UAutomatedPerfTestControllerBase::OnInit()
 	FParse::Value(FCommandLine::Get(), TEXT("AutomatedPerfTest.DeviceProfileOverride="), DeviceProfileOverride);
 	FParse::Value(FCommandLine::Get(), TEXT("AutomatedPerfTest.TestID="), TestID);
 	
-	if(FParse::Param(FCommandLine::Get(), TEXT("AutomatedPerfTest.DoInsightsTrace")))
+	if (FParse::Param(FCommandLine::Get(), TEXT("AutomatedPerfTest.DoInsightsTrace")))
 	{
 		UE_LOG(LogAutomatedPerfTest, Log, TEXT("Insights Trace Requested"))
 		bRequestsInsightsTrace = true;
 	}
-	if(FParse::Param(FCommandLine::Get(), TEXT("AutomatedPerfTest.DoCSVProfiler")))
+	if (FParse::Param(FCommandLine::Get(), TEXT("AutomatedPerfTest.DoCSVProfiler")))
 	{
 		UE_LOG(LogAutomatedPerfTest, Log, TEXT("CSV Profiler Requested"))
 		bRequestsCSVProfiler = true;
 	}
-	if(FParse::Param(FCommandLine::Get(), TEXT("AutomatedPerfTest.DoFPSChart")))
+	if (FParse::Param(FCommandLine::Get(), TEXT("AutomatedPerfTest.DoFPSChart")))
 	{
 		UE_LOG(LogAutomatedPerfTest, Log, TEXT("FPSCharts Requested"))
 		bRequestsFPSChart = true;
 	}
-	if(FParse::Param(FCommandLine::Get(), TEXT("AutomatedPerfTest.DoVideoCapture")))
+	if (FParse::Param(FCommandLine::Get(), TEXT("AutomatedPerfTest.DoVideoCapture")))
 	{
 		UE_LOG(LogAutomatedPerfTest, Log, TEXT("Video Capture Requested"))
 		bRequestsVideoCapture = true;
 	}
+	if (FParse::Param(FCommandLine::Get(), TEXT("AutomatedPerfTest.LockDynamicRes")))
+	{
+		UE_LOG(LogAutomatedPerfTest, Log, TEXT("Locking dynamic res requested"))
+		bRequestsLockedDynRes = true;
+	}
 	
 	FWorldDelegates::OnPreWorldInitialization.AddUObject(this, &ThisClass::OnPreWorldInitializeInternal);
 }
diff --git a/Engine/Plugins/Performance/AutomatedPerfTesting/Source/AutomatedPerfTesting/Public/AutomatedPerfTestControllerBase.h b/Engine/Plugins/Performance/AutomatedPerfTesting/Source/AutomatedPerfTesting/Public/AutomatedPerfTestControllerBase.h
--- a/Engine/Plugins/Performance/AutomatedPerfTesting/Source/AutomatedPerfTesting/Public/AutomatedPerfTestControllerBase.h
+++ b/Engine/Plugins/Performance/AutomatedPerfTesting/Source/AutomatedPerfTesting/Public/AutomatedPerfTestControllerBase.h
@@ -81,6 +81,7 @@ public:
 	bool RequestsCSVProfiler() const;
 	bool RequestsFPSChart() const;
 	bool RequestsVideoCapture() const;
+	bool RequestsLockedDynRes() const;
 
 	bool TryStartInsightsTrace();
 	bool TryStopInsightsTrace();
@@ -137,6 +138,7 @@ private:
 	bool bRequestsInsightsTrace;
 	bool bRequestsCSVProfiler;
 	bool bRequestsVideoCapture;
+	bool bRequestsLockedDynRes;
 
 	FText VideoRecordingTitle;
 	
