Use EngineDir property from UnrealEngine.csproj.props

#rnx

diff --git a/Engine/Plugins/Performance/AutomatedPerfTesting/Build/Scripts/AutomatedPerfTesting.Automation.csproj b/Engine/Plugins/Performance/AutomatedPerfTesting/Build/Scripts/AutomatedPerfTesting.Automation.csproj
--- a/Engine/Plugins/Performance/AutomatedPerfTesting/Build/Scripts/AutomatedPerfTesting.Automation.csproj
+++ b/Engine/Plugins/Performance/AutomatedPerfTesting/Build/Scripts/AutomatedPerfTesting.Automation.csproj
@@ -1,9 +1,6 @@
 <Project Sdk="Microsoft.NET.Sdk">
   <Import Project="AutomatedPerfTesting.Automation.csproj.props" Condition="Exists('AutomatedPerfTesting.Automation.csproj.props')" />
-  <PropertyGroup>
-    <EngineDir>../../../../..</EngineDir>
-  </PropertyGroup>
-  <Import Project="$(EngineDir)\Source\Programs\Shared\UnrealEngine.csproj.props" />
+  <Import Project="..\..\..\..\..\Source\Programs\Shared\UnrealEngine.csproj.props" />
   <PropertyGroup>
     <TargetFramework>net8.0</TargetFramework>
     <OutputType>Library</OutputType>
