BATP - add support for multiple RHIs

diff --git a/Engine/Plugins/Performance/AutomatedPerfTesting/Build/Inc/AutomatedPerfTestCommonSettings.xml b/Engine/Plugins/Performance/AutomatedPerfTesting/Build/Inc/AutomatedPerfTestCommonSettings.xml
--- a/Engine/Plugins/Performance/AutomatedPerfTesting/Build/Inc/AutomatedPerfTestCommonSettings.xml
+++ b/Engine/Plugins/Performance/AutomatedPerfTesting/Build/Inc/AutomatedPerfTestCommonSettings.xml
@@ -6,6 +6,9 @@
 	<!-- insights traces after the conclusion of the test --> 
 	<Option Name="UseShippingInsights" Restrict="true|false" DefaultValue="false" Description="Whether or not to use the shipping configuration of Insights." />
 
+	<!-- Test configurations are more expressive way of determining of what variations of things need to be tested -->
+	<Include Script="TestConfigs.xml" />
+
 	<!-- then we'll use that to override values that will be picked up by the rest of BuildAndTestProject.xml -->
 	<Property Name="IncludeUnrealInsightsInToolCompile" Value="true" />
 	<Property Name="UnrealInsightsBuildConfiguration" Value="Development" If="!$(UseShippingInsights)" />
@@ -38,6 +41,7 @@
 	<Option Name="WithAPTStaticCameraTests" DefaultValue="false" Restrict="true|false" Description="Whether to run Auto Perf Test static camera tests."  />
 	<Option Name="WithAPTMaterialTest" DefaultValue="false" Restrict="true|false" Description="Whether to run Auto Perf Test material test."  />
 	<Option Name="WithAPTReplayTest"   DefaultValue="false" Restrict="true|false" Description="Whether to run Auto Perf Test replay test."  />
+	<Option Name="WithTestConfigs" DefaultValue="false" Restrict="true|false" Description="Whether or not to use the additional test configs" />
 
 	<!-- Test sub-types, act as flavors to the above tests -->
 	<Option Name="DoPerf" DefaultValue="false" Restrict="true|false" Description="Whether to run Perf (general performance) test." />
@@ -71,37 +75,66 @@
 		<Property Name="InsightsTestCommonArgs" Value="$(InsightsTestCommonArgs) -AutomatedPerfTest.UsePlatformExtraDebugMemory"/>
 	</Do>
 
-	<!-- append plugin-provided tests to GauntletSettings' TargetTestList -->
-	<Do If="$(WithAPTSequenceTests)">
-		<ForEach Name="MapSequenceComboName" Values="$(MapSequenceComboNames)" Separator="+">
-			<Property Name="TargetTestList" Value="$(TargetTestList)AutomatedPerfTest.SequenceTest{-AutomatedPerfTest.SequencePerfTest.MapSequenceName=$(MapSequenceComboName) -AutomatedPerfTest.DoCSVProfiler $(PerfTestCommonArgs)}[Sequence Perf $(MapSequenceComboName)]+" If="$(DoPerf)" />
-			<Property Name="TargetTestList" Value="$(TargetTestList)AutomatedPerfTest.SequenceTest{-AutomatedPerfTest.SequencePerfTest.MapSequenceName=$(MapSequenceComboName) -AutomatedPerfTest.DoCSVProfiler -AutomatedPerfTest.DoLLM $(LLMTestCommonArgs)}[Sequence LLM $(MapSequenceComboName)]+" If="$(DoLLM)"/>
-			<Property Name="TargetTestList" Value="$(TargetTestList)AutomatedPerfTest.SequenceTest{-AutomatedPerfTest.SequencePerfTest.MapSequenceName=$(MapSequenceComboName) -AutomatedPerfTest.DoCSVProfiler -AutomatedPerfTest.DoGPUPerf $(GPUPerfTestCommonArgs)}[Sequence GPUPerf $(MapSequenceComboName)]+" If="$(DoGPUPerf)" />
-			<Property Name="TargetTestList" Value="$(TargetTestList)AutomatedPerfTest.SequenceTest{-AutomatedPerfTest.SequencePerfTest.MapSequenceName=$(MapSequenceComboName) -AutomatedPerfTest.DoCSVProfiler -AutomatedPerfTest.DoInsightsTrace $(InsightsTestCommonArgs)}[Sequence Insights $(MapSequenceComboName)]+" If="$(DoInsightsTrace)"/>
-		</ForEach>
-	</Do>
+	<Property Name="AllTestConfigs" Value="" If="!$(WithTestConfigs)" />
 
-	<Do If="$(WithAPTStaticCameraTests)">
-		<ForEach Name="StaticCameraMapName" Values="$(StaticCameraMapNames)" Separator="+">
-			<Property Name="TargetTestList" Value="$(TargetTestList)AutomatedPerfTest.StaticCameraTest{-AutomatedPerfTest.StaticCameraPerfTest.MapName=$(StaticCameraMapName) -AutomatedPerfTest.DoCSVProfiler $(PerfTestCommonArgs)}[Static Camera Perf $(StaticCameraMapName)]+" If="$(DoPerf)" />
-			<Property Name="TargetTestList" Value="$(TargetTestList)AutomatedPerfTest.StaticCameraTest{-AutomatedPerfTest.StaticCameraPerfTest.MapName=$(StaticCameraMapName) -AutomatedPerfTest.DoCSVProfiler -AutomatedPerfTest.DoLLM $(LLMTestCommonArgs)}[Static Camera LLM $(StaticCameraMapName)]+" If="$(DoLLM)"/>
-			<Property Name="TargetTestList" Value="$(TargetTestList)AutomatedPerfTest.StaticCameraTest{-AutomatedPerfTest.StaticCameraPerfTest.MapName=$(StaticCameraMapName) -AutomatedPerfTest.DoCSVProfiler -AutomatedPerfTest.DoGPUPerf $(LLMTestCommonArgs)}[Static Camera GPUPerf $(StaticCameraMapName)]+" If="$(DoGPUPerf)"/>
-			<Property Name="TargetTestList" Value="$(TargetTestList)AutomatedPerfTest.StaticCameraTest{-AutomatedPerfTest.StaticCameraPerfTest.MapName=$(StaticCameraMapName) -AutomatedPerfTest.DoCSVProfiler -AutomatedPerfTest.DoInsightsTrace $(InsightsTestCommonArgs)}[Static Camera Insights $(StaticCameraMapName)]+" If="$(DoInsightsTrace)"/>
-		</ForEach>
-	</Do>
+	<!-- Add in the base test config -->
+	<Property Name="AllTestConfigs" Value=".|.|.|.;$(AllTestConfigs)" />
 
-	<Do If="$(WithAPTMaterialTest)">
-		<Property Name="TargetTestList" Value="$(TargetTestList)AutomatedPerfTest.MaterialTest{-AutomatedPerfTest.DoCSVProfiler $(PerfTestCommonArgs)}[Material Perf]+" If="$(DoPerf)" />
-		<Property Name="TargetTestList" Value="$(TargetTestList)AutomatedPerfTest.MaterialTest{-AutomatedPerfTest.DoCSVProfiler -AutomatedPerfTest.DoLLM $(LLMTestCommonArgs)}[Material LLM]+" If="$(DoLLM)"/>
-		<Property Name="TargetTestList" Value="$(TargetTestList)AutomatedPerfTest.MaterialTest{-AutomatedPerfTest.DoCSVProfiler -AutomatedPerfTest.DoGPUPerf $(LLMTestCommonArgs)}[Material GPUPerf]+" If="$(DoGPUPerf)"/>
-		<Property Name="TargetTestList" Value="$(TargetTestList)AutomatedPerfTest.MaterialTest{-AutomatedPerfTest.DoCSVProfiler -AutomatedPerfTest.DoInsightsTrace $(InsightsTestCommonArgs)}[Material Insights]+" If="$(DoInsightsTrace)"/>
-	</Do>
+	<ForEach Name="TestConfig" Values="$(AllTestConfigs)" Separator=";">
+		<Regex Input="$(TestConfig)" Pattern="(.*)\|(.*)\|(.*)\|(.*)" Capture="PlatformToTest;TestConfigName;ArtifactPostfix;AdditionalArguments"/>
 
-	<Do If="$(WithAPTReplayTest)">
-		<Property Name="TargetTestList" Value="$(TargetTestList)AutomatedPerfTest.ReplayTest{-AutomatedPerfTest.ReplayPerfTest.ReplayName=&quot;$(ReplayName)&quot; -AutomatedPerfTest.DoCSVProfiler $(PerfTestCommonArgs)}[Replay Perf]+" If="$(DoPerf)" />
-		<Property Name="TargetTestList" Value="$(TargetTestList)AutomatedPerfTest.ReplayTest{-AutomatedPerfTest.ReplayPerfTest.ReplayName=&quot;$(ReplayName)&quot; -AutomatedPerfTest.DoCSVProfiler -AutomatedPerfTest.DoLLM $(LLMTestCommonArgs)}[Replay LLM]+" If="$(DoLLM)"/>
-		<Property Name="TargetTestList" Value="$(TargetTestList)AutomatedPerfTest.ReplayTest{-AutomatedPerfTest.ReplayPerfTest.ReplayName=&quot;$(ReplayName)&quot; -AutomatedPerfTest.DoCSVProfiler -AutomatedPerfTest.DoGPUPerf $(LLMTestCommonArgs)}[Replay GPUPerf]+" If="$(DoGPUPerf)"/>
-		<Property Name="TargetTestList" Value="$(TargetTestList)AutomatedPerfTest.ReplayTest{-AutomatedPerfTest.ReplayPerfTest.ReplayName=&quot;$(ReplayName)&quot; -AutomatedPerfTest.DoCSVProfiler -AutomatedPerfTest.DoInsightsTrace $(InsightsTestCommonArgs)}[Replay Insights]+" If="$(DoInsightsTrace)"/>
-	</Do>
+		<Do If="'$(TestConfigName)' == '.'">
+			<Property Name="TestConfigName" Value="" />
+		</Do>
+		<Property Name="TestConfigArgs" Value=" -AutomatedPerfTest.TestConfigName=&quot;$(TestConfigName)&quot;" />
+
+		<Do If="'$(TestConfigName)' != ''">
+			<Property Name="TestConfigName" Value=" $(TestConfigName)" />
+		</Do>
+
+		<Do If="'$(ArtifactPostfix)' != '.'">
+			<Property Name="TestConfigArgs" Value="$(TestConfigArgs) -ArtifactPostfix=$(ArtifactPostfix)" />
+		</Do>
+
+		<Do If="'$(AdditionalArguments)' != '.'">
+			<Property Name="TestConfigArgs" Value="$(TestConfigArgs) -args=&quot;$(AdditionalArguments)&quot;" />
+		</Do>
+
+		<!-- The special sentinel value of '.' denotes to always run the base test, this will change once we refactor to always use TestConfigs -->
+		<Do If="ContainsItem('$(PlatformToTest)', '$(TargetPlatforms)', '+') or '$(PlatformToTest)' == '.'">
+			<!-- append plugin-provided tests to GauntletSettings' TargetTestList -->
+			<Do If="$(WithAPTSequenceTests)">
+				<ForEach Name="MapSequenceComboName" Values="$(MapSequenceComboNames)" Separator="+">
+					<Property Name="TargetTestList" Value="$(TargetTestList)AutomatedPerfTest.SequenceTest{-AutomatedPerfTest.SequencePerfTest.MapSequenceName=$(MapSequenceComboName) -AutomatedPerfTest.DoCSVProfiler $(PerfTestCommonArgs)$(TestConfigArgs)}[Sequence Perf $(MapSequenceComboName)$(TestConfigName)]+" If="$(DoPerf)" />
+					<Property Name="TargetTestList" Value="$(TargetTestList)AutomatedPerfTest.SequenceTest{-AutomatedPerfTest.SequencePerfTest.MapSequenceName=$(MapSequenceComboName) -AutomatedPerfTest.DoCSVProfiler -AutomatedPerfTest.DoLLM $(LLMTestCommonArgs)$(TestConfigArgs)}[Sequence LLM $(MapSequenceComboName)$(TestConfigName)]+" If="$(DoLLM)"/>
+					<Property Name="TargetTestList" Value="$(TargetTestList)AutomatedPerfTest.SequenceTest{-AutomatedPerfTest.SequencePerfTest.MapSequenceName=$(MapSequenceComboName) -AutomatedPerfTest.DoCSVProfiler -AutomatedPerfTest.DoGPUPerf $(GPUPerfTestCommonArgs)$(TestConfigArgs)}[Sequence GPUPerf $(MapSequenceComboName)$(TestConfigName)]+" If="$(DoGPUPerf)" />
+					<Property Name="TargetTestList" Value="$(TargetTestList)AutomatedPerfTest.SequenceTest{-AutomatedPerfTest.SequencePerfTest.MapSequenceName=$(MapSequenceComboName) -AutomatedPerfTest.DoCSVProfiler -AutomatedPerfTest.DoInsightsTrace $(InsightsTestCommonArgs)$(TestConfigArgs)}[Sequence Insights $(MapSequenceComboName)$(TestConfigName)]+" If="$(DoInsightsTrace)"/>
+				</ForEach>
+			</Do>
+
+			<Do If="$(WithAPTStaticCameraTests)">
+				<ForEach Name="StaticCameraMapName" Values="$(StaticCameraMapNames)" Separator="+">
+					<Property Name="TargetTestList" Value="$(TargetTestList)AutomatedPerfTest.StaticCameraTest{-AutomatedPerfTest.StaticCameraPerfTest.MapName=$(StaticCameraMapName) -AutomatedPerfTest.DoCSVProfiler $(PerfTestCommonArgs)$(TestConfigArgs)}[Static Camera Perf $(StaticCameraMapName)$(TestConfigName)]+" If="$(DoPerf)" />
+					<Property Name="TargetTestList" Value="$(TargetTestList)AutomatedPerfTest.StaticCameraTest{-AutomatedPerfTest.StaticCameraPerfTest.MapName=$(StaticCameraMapName) -AutomatedPerfTest.DoCSVProfiler -AutomatedPerfTest.DoLLM $(LLMTestCommonArgs)$(TestConfigArgs)}[Static Camera LLM $(StaticCameraMapName)$(TestConfigName)]+" If="$(DoLLM)"/>
+					<Property Name="TargetTestList" Value="$(TargetTestList)AutomatedPerfTest.StaticCameraTest{-AutomatedPerfTest.StaticCameraPerfTest.MapName=$(StaticCameraMapName) -AutomatedPerfTest.DoCSVProfiler -AutomatedPerfTest.DoGPUPerf $(LLMTestCommonArgs)$(TestConfigArgs)}[Static Camera GPUPerf $(StaticCameraMapName)$(TestConfigName)]+" If="$(DoGPUPerf)"/>
+					<Property Name="TargetTestList" Value="$(TargetTestList)AutomatedPerfTest.StaticCameraTest{-AutomatedPerfTest.StaticCameraPerfTest.MapName=$(StaticCameraMapName) -AutomatedPerfTest.DoCSVProfiler -AutomatedPerfTest.DoInsightsTrace $(InsightsTestCommonArgs)$(TestConfigArgs)}[Static Camera Insights $(StaticCameraMapName)$(TestConfigName)]+" If="$(DoInsightsTrace)"/>
+				</ForEach>
+			</Do>
+
+			<Do If="$(WithAPTMaterialTest)">
+				<Property Name="TargetTestList" Value="$(TargetTestList)AutomatedPerfTest.MaterialTest{-AutomatedPerfTest.DoCSVProfiler $(PerfTestCommonArgs)$(TestConfigArgs)}[Material Perf$(TestConfigName)]+" If="$(DoPerf)" />
+				<Property Name="TargetTestList" Value="$(TargetTestList)AutomatedPerfTest.MaterialTest{-AutomatedPerfTest.DoCSVProfiler -AutomatedPerfTest.DoLLM $(LLMTestCommonArgs)$(TestConfigArgs)}[Material LLM$(TestConfigName)]+" If="$(DoLLM)"/>
+				<Property Name="TargetTestList" Value="$(TargetTestList)AutomatedPerfTest.MaterialTest{-AutomatedPerfTest.DoCSVProfiler -AutomatedPerfTest.DoGPUPerf $(LLMTestCommonArgs)$(TestConfigArgs)}[Material GPUPerf$(TestConfigName)]+" If="$(DoGPUPerf)"/>
+				<Property Name="TargetTestList" Value="$(TargetTestList)AutomatedPerfTest.MaterialTest{-AutomatedPerfTest.DoCSVProfiler -AutomatedPerfTest.DoInsightsTrace $(InsightsTestCommonArgs)$(TestConfigArgs)}[Material Insights$(TestConfigName)]+" If="$(DoInsightsTrace)"/>
+			</Do>
+
+			<Do If="$(WithAPTReplayTest)">
+				<Property Name="TargetTestList" Value="$(TargetTestList)AutomatedPerfTest.ReplayTest{-AutomatedPerfTest.ReplayPerfTest.ReplayName=&quot;$(ReplayName)&quot; -AutomatedPerfTest.DoCSVProfiler $(PerfTestCommonArgs)$(TestConfigArgs)}[Replay Perf$(TestConfigName)]+" If="$(DoPerf)" />
+				<Property Name="TargetTestList" Value="$(TargetTestList)AutomatedPerfTest.ReplayTest{-AutomatedPerfTest.ReplayPerfTest.ReplayName=&quot;$(ReplayName)&quot; -AutomatedPerfTest.DoCSVProfiler -AutomatedPerfTest.DoLLM $(LLMTestCommonArgs)$(TestConfigArgs)}[Replay LLM$(TestConfigName)]+" If="$(DoLLM)"/>
+				<Property Name="TargetTestList" Value="$(TargetTestList)AutomatedPerfTest.ReplayTest{-AutomatedPerfTest.ReplayPerfTest.ReplayName=&quot;$(ReplayName)&quot; -AutomatedPerfTest.DoCSVProfiler -AutomatedPerfTest.DoGPUPerf $(LLMTestCommonArgs)$(TestConfigArgs)}[Replay GPUPerf$(TestConfigName)]+" If="$(DoGPUPerf)"/>
+				<Property Name="TargetTestList" Value="$(TargetTestList)AutomatedPerfTest.ReplayTest{-AutomatedPerfTest.ReplayPerfTest.ReplayName=&quot;$(ReplayName)&quot; -AutomatedPerfTest.DoCSVProfiler -AutomatedPerfTest.DoInsightsTrace $(InsightsTestCommonArgs)$(TestConfigArgs)}[Replay Insights$(TestConfigName)]+" If="$(DoInsightsTrace)"/>
+			</Do>
+		</Do>
+	</ForEach>
 
 </BuildGraph>
diff --git a/Engine/Plugins/Performance/AutomatedPerfTesting/Build/Inc/TestConfigs.xml b/Engine/Plugins/Performance/AutomatedPerfTesting/Build/Inc/TestConfigs.xml
new file mode 100644
--- /dev/null
+++ b/Engine/Plugins/Performance/AutomatedPerfTesting/Build/Inc/TestConfigs.xml
@@ -0,0 +1,38 @@
+<?xml version='1.0' ?>
+<BuildGraph xmlns="http://www.epicgames.com/BuildGraph" xmlns:xsi="http://www.w3.org/2001/XMLSchema-instance" xsi:schemaLocation="http://www.epicgames.com/BuildGraph ../../../../Engine/Build/Graph/Schema.xsd" >
+
+	<Property Name="AllTestConfigs" Value=""/>
+	
+	<Macro Name="AddTestConfig" Arguments="TestPlatform;TestConfigName;ArtifactPostfix;AdditionalArguments">
+		<Property Name="AllTestConfigs" Value="$(AllTestConfigs);$(TestPlatform)|$(TestConfigName)|$(ArtifactPostfix)|$(AdditionalArguments);"/>
+	</Macro>
+
+	<Expand Name="AddTestConfig" 
+		TestPlatform="Win64"
+		TestConfigName="Vulkan SM6"
+		ArtifactPostfix="vulkan_sm6"
+		AdditionalArguments="-vulkan -sm6"
+	/>
+
+	<Expand Name="AddTestConfig" 
+		TestPlatform="Linux"
+		TestConfigName="Vulkan SM5"
+		ArtifactPostfix="vulkan_sm5"
+		AdditionalArguments="-vulkan -sm5"
+	/>
+
+	<Expand Name="AddTestConfig" 
+		TestPlatform="Win64"
+		TestConfigName="Vulkan SM5"
+		ArtifactPostfix="vulkan_sm5"
+		AdditionalArguments="-vulkan -sm5"
+	/>
+
+	<Expand Name="AddTestConfig" 
+		TestPlatform="Linux"
+		TestConfigName="Vulkan ES31"
+		ArtifactPostfix="vulkan_es31"
+		AdditionalArguments="-vulkan -es31"
+	/>
+
+</BuildGraph>
diff --git a/Engine/Plugins/Performance/AutomatedPerfTesting/Build/Scripts/AutomatedPerfTestConfig.cs b/Engine/Plugins/Performance/AutomatedPerfTesting/Build/Scripts/AutomatedPerfTestConfig.cs
--- a/Engine/Plugins/Performance/AutomatedPerfTesting/Build/Scripts/AutomatedPerfTestConfig.cs
+++ b/Engine/Plugins/Performance/AutomatedPerfTesting/Build/Scripts/AutomatedPerfTestConfig.cs
@@ -137,6 +137,12 @@ namespace AutomatedPerfTest
 		[AutoParamWithNames(false, "AutomatedPerfTest.IgnoreTestBuildLogging")]
 		public bool IgnoreTestBuildLogging;
 
+		/// <summary>
+		/// Used to specify a different test name in the CSV meta data, primarily used to distinguish between RHIs
+		/// </summary>
+		[AutoParamWithNames("", "AutomatedPerfTest.TestConfigName")]
+		public string TestConfigName = "";
+
 		public string DataSourceName;
 
 		/// <summary>
diff --git a/Engine/Plugins/Performance/AutomatedPerfTesting/Build/Scripts/AutomatedPerfTestNode.cs b/Engine/Plugins/Performance/AutomatedPerfTesting/Build/Scripts/AutomatedPerfTestNode.cs
--- a/Engine/Plugins/Performance/AutomatedPerfTesting/Build/Scripts/AutomatedPerfTestNode.cs
+++ b/Engine/Plugins/Performance/AutomatedPerfTesting/Build/Scripts/AutomatedPerfTestNode.cs
@@ -711,6 +711,7 @@ namespace AutomatedPerfTest
 				CsvMetadata.Add($"gauntletSubTest={GetSubtestName(Config)}");
 				CsvMetadata.Add("testBuildIsPreflight=" + (ReportGenUtils.IsTestingPreflightBuild(OriginalBuildName) ? "1" : "0"));
 				CsvMetadata.Add($"testBuildVersion={OriginalBuildName}");
+				CsvMetadata.Add("testconfigname="+Config.TestConfigName);
 
 				if (!string.IsNullOrEmpty(Context.BuildInfo.Branch) && Context.BuildInfo.Changelist != 0)
 				{
