Making use of Gauntlet's logger in order for issues to be better consolidated.

#rnx

diff --git a/Engine/Plugins/Performance/AutomatedPerfTesting/Build/Scripts/AutomatedPerfTestNode.cs b/Engine/Plugins/Performance/AutomatedPerfTesting/Build/Scripts/AutomatedPerfTestNode.cs
--- a/Engine/Plugins/Performance/AutomatedPerfTesting/Build/Scripts/AutomatedPerfTestNode.cs
+++ b/Engine/Plugins/Performance/AutomatedPerfTesting/Build/Scripts/AutomatedPerfTestNode.cs
@@ -2,15 +2,12 @@
 
 using System;
 using System.Collections.Generic;
-using System.Runtime.InteropServices;
 using System.IO;
 using System.Linq;
-using EpicGame;
 using AutomationTool;
 using Gauntlet;
 using EpicGames.Core;
 using Log = Gauntlet.Log;
-using Microsoft.Extensions.Logging;
 using UnrealBuildBase;	// for Unreal.RootDirectory
 using UnrealBuildTool;	// for UnrealTargetPlatform
 
@@ -277,7 +274,7 @@ namespace AutomatedPerfTest
 			// failures if they wish to do so. 
 			if(!InLogSummary.EngineInitialized && TargetConfig == UnrealTargetConfiguration.Test && ProcessExitCode == 0)
 			{
-				Logger.LogWarning("*** Engine Initialization log not detected in Test build with Exit Code 0. " +
+				Log.Warning("*** Engine Initialization log not detected in Test build with Exit Code 0. " +
 					"This test will fail. " +
 					"Try passing `-AutomatedPerfTest.IgnoreTestBuildLogging` while running this test " +
 					"or pass `-set:APTIgnoreTestBuildLogging=true` if running via BuildGraph if you " +
@@ -330,7 +327,7 @@ namespace AutomatedPerfTest
 						Log.Info("Creating perf server importer with build name {BuildName}", OriginalBuildName);
 						string DataSourceName = GetConfiguration().DataSourceName;
 						string ImportDirOverride = Globals.Params.ParseValue("PerfReportServerImportDir", null);
-						Logger.LogInformation("Creating PRS Importer for data source '{0}' and import dir override (if any) '{1}'.",
+						Log.Info("Creating PRS Importer for data source '{0}' and import dir override (if any) '{1}'.",
 							DataSourceName, ImportDirOverride);
 						ICsvImporter Importer = ReportGenUtils.CreatePerfReportServerImporter(DataSourceName, OriginalBuildName,
 							IsBuildMachine, ImportDirOverride, CommonDataSourceFields);
@@ -396,7 +393,7 @@ namespace AutomatedPerfTest
 						}
 						else
 						{
-							Logger.LogWarning("Unable to create PRS Importer.");
+							Log.Warning("Unable to create PRS Importer.");
 						}
 
 						// Cleanup the temp dir
@@ -408,13 +405,13 @@ namespace AutomatedPerfTest
 
 					if (!bGeneratedLocalReport && !bGeneratedPRSReport)
 					{
-						Logger.LogWarning("Did not generate neither local nor a PRS report.");
+						Log.Warning("Did not generate neither local nor a PRS report.");
 					}
 				}
 			}
 			else
 			{
-				Logger.LogWarning("Skipping performance report generation because the perf report test failed.");
+				Log.Warning("Skipping performance report generation because the perf report test failed.");
 			}
 
 			return base.CreateReport(Result, Context, Build, Artifacts, ArtifactPath);
@@ -458,7 +455,7 @@ namespace AutomatedPerfTest
 			var ToolPath = FileReference.Combine(Unreal.EngineDirectory, "Binaries", "DotNET", "CsvTools", perfreportTool);
 			if (!FileReference.Exists(ToolPath))
 			{
-				Logger.LogError("Failed to find perf report utility at this path: \"{ToolPath}\".", ToolPath);
+				Log.Error("Failed to find perf report utility at this path: \"{ToolPath}\".", ToolPath);
 				return false;
 			}
 
@@ -473,7 +470,7 @@ namespace AutomatedPerfTest
 			if(string.IsNullOrEmpty(ReportPath))
 			{
 				ReportPath = Path.Combine(OutputPath, "Reports");
-			};
+			}
 
 			string ReportCacheDir = Path.Combine(OutputPath, "Cache");
 
@@ -497,7 +494,7 @@ namespace AutomatedPerfTest
 
 			if (DiscoveredCsvs.Count == 0)
 			{
-				Logger.LogError("Test completed successfully but no csv profiling results were found. Searched paths were:\r\n  {Paths}", string.Join("\r\n  ", CsvsPaths.Select(s => $"\"{s}\"")));
+				Log.Error("Test completed successfully but no csv profiling results were found. Searched paths were:\r\n  {Paths}", string.Join("\r\n  ", CsvsPaths.Select(s => $"\"{s}\"")));
 				return false;
 			}
 
@@ -520,7 +517,7 @@ namespace AutomatedPerfTest
 				try { Directory.CreateDirectory(ReportCacheDir); }
 				catch (Exception Ex)
 				{
-					Logger.LogError("Failed to create perf report cache directory \"{ReportCacheDir}\". {Ex}", ReportCacheDir, Ex);
+					Log.Error("Failed to create perf report cache directory \"{ReportCacheDir}\". {Ex}", ReportCacheDir, Ex);
 					return false;
 				}
 			}
@@ -529,7 +526,7 @@ namespace AutomatedPerfTest
 				try { Directory.CreateDirectory(ReportPath); }
 				catch (Exception Ex)
 				{
-					Logger.LogError("Failed to create perf report output directory \"{ReportPath}\". {Ex}", ReportPath, Ex);
+					Log.Error("Failed to create perf report output directory \"{ReportPath}\". {Ex}", ReportPath, Ex);
 					return false;
 				}
 			}
@@ -562,7 +559,7 @@ namespace AutomatedPerfTest
 			RunAndLog(CmdEnv, CmdEnv.DotnetMsbuildPath, PerfReportToolArgs, out int ErrorCode);
 			if (ErrorCode != 0)
 			{
-				Logger.LogError("PerfReportTool returned error code \"{ErrorCode}\" while generating detailed report.", ErrorCode);
+				Log.Error("PerfReportTool returned error code \"{ErrorCode}\" while generating detailed report.", ErrorCode);
 			}
 
 			// Now generate the all-time historic summary report
@@ -608,7 +605,7 @@ namespace AutomatedPerfTest
 				RunAndLog(CmdEnv, CmdEnv.DotnetMsbuildPath, ArgStr, out ErrorCode);
 				if (ErrorCode != 0)
 				{
-					Logger.LogError("PerfReportTool returned error code \"{ErrorCode}\" while generating historic report.", ErrorCode);
+					Log.Error("PerfReportTool returned error code \"{ErrorCode}\" while generating historic report.", ErrorCode);
 				}
 			}
 
@@ -829,7 +826,7 @@ namespace AutomatedPerfTest
 		
 		public void CopyInsightsTraceToOutput(string FromArtifactPath, string ToOutputPath)
 		{
-			Logger.LogInformation("Copying test insights trace from artifact path to report cache");
+			Log.Info("Copying test insights trace from artifact path to report cache");
 			
 			// find all the available trace paths
 			var DiscoveredTraces = new List<string>();
@@ -843,7 +840,7 @@ namespace AutomatedPerfTest
 			// if we couldn't find any traces, report that and bail out
 			if (DiscoveredTraces.Count == 0)
 			{
-				Logger.LogError("Test completed successfully but no trace results were found. Searched path was {ArtifactPath}", FromArtifactPath);
+				Log.Error("Test completed successfully but no trace results were found. Searched path was {ArtifactPath}", FromArtifactPath);
 				return;
 			}
 			
@@ -860,8 +857,8 @@ namespace AutomatedPerfTest
 			{
 				string Filename = Path.GetFileNameWithoutExtension(SortedTraces[0]);
 				string PerfTracePath = Path.Combine(ReportPath, Filename + ".utrace");
-				
-				Logger.LogInformation("Copying latest utrace file from {ArtifactPath} to Perf .utrace path: {PerfTracePath}", FromArtifactPath,
+
+				Log.Info("Copying latest utrace file from {ArtifactPath} to Perf .utrace path: {PerfTracePath}", FromArtifactPath,
 					PerfTracePath);
 				
 				// just try the copy over, and log a failure, but don't bail out of the test.
@@ -872,7 +869,7 @@ namespace AutomatedPerfTest
 				}
 				catch (Exception e)
 				{
-					Logger.LogWarning("Failed to copy local trace file: {Text}", e);
+					Log.Warning("Failed to copy local trace file: {Text}", e);
 				}
 			}
 		}
