Support the concept of `TestConfig`'s that change the way test nodes are generated and processed.
Currently this is opt-in and only Supernova (an HMI benchmark project) uses it

A TestConfig has a few fields:
- Platform to run on
- The name of the config (this should be stable so PRS can determine which runs to group together)
- Which subfolder the artifacts should be written to (ensures that PRS keeps the performance data separate per config)
- Additional arguments to inject to the game (e.g. -vulkan -sm6)
- Which Agent Type this config should run on (This is a requirement for Mac hardware when testing against SM5 and SM6)

Instead of using the current loop inside of BuildAndTestProject for generating test nodes, we have a custom block of logic that determines how to generate these things
- This is because BuildAndTestProject assumes that all Test Targets should be on every platform
- With TestConfig that is no longer a valid assumption

diff --git a/Engine/Build/Graph/Tasks/BuildAndTestProject.xml b/Engine/Build/Graph/Tasks/BuildAndTestProject.xml
--- a/Engine/Build/Graph/Tasks/BuildAndTestProject.xml
+++ b/Engine/Build/Graph/Tasks/BuildAndTestProject.xml
@@ -170,6 +170,8 @@
 		<Include Script="$(PreStepsExtendFile)"/>
 	</Do>
 
+	<Option Name="UseTestConfigs" DefaultValue="false" Description="Use Test Configs (see RunTestConfigs.xml) instead of the standard test node generation"/>
+
 	<!-- At this point we have all options so calculate a few things -->
 
 	<!-- Validate that the project exists-->
@@ -522,6 +524,14 @@
 		<Label Category="Editors" Name="$(PreLabelName)CookedEditor $(EditorPlatform)" Requires="$(GeneratedCookedEditorNodes)" />
 	</ForEach>
 
+	<ForEach Name="TargetPlatform" Values="$(TargetPlatforms)" Separator="+">
+		<!-- We cannot have an inner loop or these won't be added to the parent scope -->
+		<Property Name="TestPrerequisites$(TargetPlatform)Shipping" Value="" CreateInParentScope="true"/>
+		<Property Name="TestPrerequisites$(TargetPlatform)Test" Value="" CreateInParentScope="true" />
+		<Property Name="TestPrerequisites$(TargetPlatform)Development" Value="" CreateInParentScope="true" />
+		<Property Name="TestPrerequisites$(TargetPlatform)Debug" Value="" CreateInParentScope="true" />
+	</ForEach>
+
 	<!-- Now, for each platform we build, cook, stage, package, publish, and test -->
 	<ForEach Name="TargetPlatform" Values="$(TargetPlatforms)" Separator="+">
 
@@ -828,29 +838,32 @@
 		<Property Name="AllPlatformInstallNodes" Value="" />
 		<Property Name="AllPlatformBootTestNodes" Value="" />
 		<Property Name="AllPlatformTestNodes" Value="" />
+
+		<Property Name="UsePackaged" Value="(ContainsItem('$(PlatformsThatRequirePackaging)', '$(TargetPlatform)', '+') Or !ContainsItem('$(PlatformsToSkipPackaging)', '$(TargetPlatform)', '+'))" />
+
+		<!-- Create the -build argument for Gauntlet based on whether builds are local or on the network-->
+		<Property Name="GauntletBuildArg" Value="-build=&quot;$(ProjectOutputDirectory)/$(StagedPlatformFolder)&quot;" />
+		<Property Name="GauntletBuildArg" Value="-build=&quot;$(ProjectOutputDirectory)/Packages/$(StagedPlatformFolder)&quot; -packaged" If="$(UsePackaged)" />
+
+		<Do If="!$(Horde) And ('$(NetworkOutputDirectory)' != '' Or '$(GeneratedBuildVersion)' != '')">
+			<Property Name="TestDirectory" Value="$(NetworkOutputDirectory)" />
+			<Property Name="TestDirectory" Value="$(GeneratedBuildDir)" If="'$(GeneratedBuildVersion)' != ''"/>
+			<Property Name="GauntletBuildArg" Value="-build=&quot;$(TestDirectory)&quot;" />
+			<Property Name="GauntletBuildArg" Value="$(GauntletBuildArg) -packaged" If="$(UsePackaged)" />
+			<Property Name="GauntletBuildArg" Value="$(GauntletBuildArg) -destlocalinstalldir=&quot;$(LocalInstallDirectory)&quot;" If="!$(IsLocalPlatform) And $(WithDeviceReservation)"/>
+		</Do>
+
+		<Property Name="GauntletBuildArg$(TargetPlatform)" Value="$(GauntletBuildArg)" CreateInParentScope="true"/>
+
 		<Agent Name="Test $(PlatformAgentName)" Type="$(TestAgentType)">
 			<!-- Run through all configs -->
 			<ForEach Name="TestConfiguration" Values="$(TestConfigurations)" Separator="+">
 				<Property Name="NodeAnnotations" Value=""/>
-
 				<!-- Prerequisites for the test to run -->
-				<Property Name="UsePackaged" Value="(ContainsItem('$(PlatformsThatRequirePackaging)', '$(TargetPlatform)', '+') Or !ContainsItem('$(PlatformsToSkipPackaging)', '$(TargetPlatform)', '+'))" />
 				<Property Name="TestPrerequisites" Value="$(RequirePlatformStaged)" />
 				<Property Name="TestPrerequisites" Value="$(RequirePlatformPackaged)" If="$(UsePackaged)"/>
 				<Property Name="TestPrerequisites" Value="$(PlatformToolsNodeName)" If="'$(GeneratedBuildVersion)' != ''"/>
 				<Property Name="TestPrerequisites" Value="$(PlatformToolsNodeName);$(PlatformEditorNodeName)" If="'$(GeneratedBuildVersion)' != '' and $(TestsRequireEditor)"/>
-				
-				<!-- Create the -build argument for Gauntlet based on whether builds are local or on the network-->
-				<Property Name="GauntletBuildArg" Value="-build=&quot;$(ProjectOutputDirectory)/$(StagedPlatformFolder)&quot;" />
-				<Property Name="GauntletBuildArg" Value="-build=&quot;$(ProjectOutputDirectory)/Packages/$(StagedPlatformFolder)&quot; -packaged" If="$(UsePackaged)" />
-
-				<Do If="!$(Horde) And ('$(NetworkOutputDirectory)' != '' Or '$(GeneratedBuildVersion)' != '')">
-					<Property Name="TestDirectory" Value="$(NetworkOutputDirectory)" />
-					<Property Name="TestDirectory" Value="$(GeneratedBuildDir)" If="'$(GeneratedBuildVersion)' != ''"/>
-					<Property Name="GauntletBuildArg" Value="-build=&quot;$(TestDirectory)&quot;" />
-					<Property Name="GauntletBuildArg" Value="$(GauntletBuildArg) -packaged" If="$(UsePackaged)" />
-					<Property Name="GauntletBuildArg" Value="$(GauntletBuildArg) -destlocalinstalldir=&quot;$(LocalInstallDirectory)&quot;" If="!$(IsLocalPlatform) And $(WithDeviceReservation)"/>
-				</Do>
 
 				<Do If="$(Horde) and '$(RetrieveUsingCloudArtifacts)' == 'false'">
 					<Property Name="HordeArtifactPlatformRetrieveWithConfig" Value="$(HordeArtifactPlatformRetrieve) $(TestConfiguration) (Horde)"/>
@@ -929,36 +942,39 @@
 					<Property Name="TestPrerequisites" Value="$(TestPrerequisites);$(BootTestName)" />
 					<Property Name="AllPlatformBootTestNodes" Value="$(AllPlatformBootTestNodes);$(BootTestName)"/>
 				</Do>
-				<!-- Go through the list of tests -->	
-				<ForEach Name="TestName" Values="$(TargetTestList)" Separator="+">
-					<Property Name="TestNameTitle" Value="$(TestName)"/>
-					<Expand Name="SanitizeTestNameTitle"/>
-					<!-- Run and report on the test-->
-					<Do If="'$(AutomationReportOutputDirectory)' != ''">
-						<Property Name="OutputFragmentPath" Value="$(TargetPlatform)/$(TestConfiguration)/$(TestNameTitle)" />
-						<Property Name="OutputReportPath" Value="$(AutomationReportOutputDirectory)/$(OutputFragmentPath)" />
-						<Property Name="GauntletReportArgs" Value="-logdir=&quot;$(OutputReportPath)&quot; -artifactname=Gauntlet -ReportExportPath=&quot;$(OutputReportPath)&quot;" />
-						<Property Name="GauntletReportArgs" Value="$(GauntletReportArgs) -ReportURL=&quot;$(AutomationReportUrl)/$(OutputFragmentPath)/&quot;" If="'$(AutomationReportUrl)' != ''" />
-					</Do>
-					<Property Name="TestNodeName" Value="$(PreNodeName)$(TargetPlatform) $(TestConfiguration) $(TestNameTitle)" />
-					<Property Name="NodeAnnotations" Value=""/>
-					<Property Name="NodeAnnotations" Value="Workflow=$(AnnotationsTarget)" If="ContainsItem('$(AnnotationsPlatforms)', '$(TargetPlatform)', ';') and ContainsItem('$(AnnotationsTests)', '$(TestName)', ';')"/>
-					<!-- Run the test. No 'Requires' as we run on the same machine locally and use the publish folder otherwise -->
-					<Property Name="TestNameBeforeParentheses" Value="$(TestName)"/>
-					<Property Name="TestArgumentsWithinParentheses" Value=""/>
-					<Expand Name="SplitCombinedTestNameAndArgs"/>
-					<Node Name="$(TestNodeName)" Requires="$(TestPrerequisites)" Annotations="$(NodeAnnotations)">
-						<Command Name="RunUnreal" Arguments="-project=$(TargetProject) -test=&quot;$(TestNameBeforeParentheses)&quot; -platform=$(TargetPlatform) -configuration=$(TestConfiguration) -target=$(TargetName) $(GauntletBuildArg) $(GauntletTestArgs) $(TestArgumentsWithinParentheses) $(GauntletReportArgs)"/>
-					</Node>
-					<!-- Add it to our list -->
-					<Do If="!ContainsItem('$(TargetPlatformsToSkipTests)', '$(TargetPlatform)', '+')">
-						<Property Name="AllPlatformTestNodes" Value="$(AllPlatformTestNodes);$(TestNodeName)"/>
-					</Do>
-				</ForEach>
-				
+
+				<Property Name="TestPrerequisites$(TargetPlatform)$(TestConfiguration)" Value="$(TestPrerequisites)" />
+
+				<Do If="!$(UseTestConfigs)">
+					<!-- Go through the list of tests -->	
+					<ForEach Name="TestName" Values="$(TargetTestList)" Separator="+">
+						<Property Name="TestNameTitle" Value="$(TestName)"/>
+						<Expand Name="SanitizeTestNameTitle"/>
+						<!-- Run and report on the test-->
+						<Do If="'$(AutomationReportOutputDirectory)' != ''">
+							<Property Name="OutputFragmentPath" Value="$(TargetPlatform)/$(TestConfiguration)/$(TestNameTitle)" />
+							<Property Name="OutputReportPath" Value="$(AutomationReportOutputDirectory)/$(OutputFragmentPath)" />
+							<Property Name="GauntletReportArgs" Value="-logdir=&quot;$(OutputReportPath)&quot; -artifactname=Gauntlet -ReportExportPath=&quot;$(OutputReportPath)&quot;" />
+							<Property Name="GauntletReportArgs" Value="$(GauntletReportArgs) -ReportURL=&quot;$(AutomationReportUrl)/$(OutputFragmentPath)/&quot;" If="'$(AutomationReportUrl)' != ''" />
+						</Do>
+						<Property Name="TestNodeName" Value="$(PreNodeName)$(TargetPlatform) $(TestConfiguration) $(TestNameTitle)" />
+						<Property Name="NodeAnnotations" Value=""/>
+						<Property Name="NodeAnnotations" Value="Workflow=$(AnnotationsTarget)" If="ContainsItem('$(AnnotationsPlatforms)', '$(TargetPlatform)', ';') and ContainsItem('$(AnnotationsTests)', '$(TestName)', ';')"/>
+						<!-- Run the test. No 'Requires' as we run on the same machine locally and use the publish folder otherwise -->
+						<Property Name="TestNameBeforeParentheses" Value="$(TestName)"/>
+						<Property Name="TestArgumentsWithinParentheses" Value=""/>
+						<Expand Name="SplitCombinedTestNameAndArgs"/>
+						<Node Name="$(TestNodeName)" Requires="$(TestPrerequisites)" Annotations="$(NodeAnnotations)">
+							<Command Name="RunUnreal" Arguments="-project=$(TargetProject) -test=&quot;$(TestNameBeforeParentheses)&quot; -platform=$(TargetPlatform) -configuration=$(TestConfiguration) -target=$(TargetName) $(GauntletBuildArg) $(GauntletTestArgs) $(TestArgumentsWithinParentheses) $(GauntletReportArgs)"/>
+						</Node>
+						<!-- Add it to our list -->
+						<Do If="!ContainsItem('$(TargetPlatformsToSkipTests)', '$(TargetPlatform)', '+')">
+							<Property Name="AllPlatformTestNodes" Value="$(AllPlatformTestNodes);$(TestNodeName)"/>
+						</Do>
+					</ForEach>
+				</Do>
 				<!-- Expand the tests here -->
 				<Expand Name="AdditionalTestNodes"/>
-				
 			</ForEach>
 		</Agent>
 		
@@ -982,10 +998,24 @@
 
 		<Property Name="GeneratedTargetNodes" Value="$(GeneratedTargetNodes);$(PlatformTasks)"/>
 		<Property Name="GeneratedTargetTestNodes" Value="$(GeneratedTargetTestNodes);$(AllPlatformBootTestNodes);$(AllPlatformTestNodes)"/>
+
+		<Property Name="GeneratedTargetTestNodes$(TargetPlatform)" Value="$(GeneratedTargetTestNodes)" CreateInParentScope="true"/>
 	</ForEach>
 
+	<Do If="$(UseTestConfigs)">
+		<Include Script="RunTestConfigs.xml"/>
+		<!-- If you have test configs set and you don't have any, include the default one. If this has issues. we can stop it doing -->
+		<Do If="'$(AllTestConfigs)' == ''">
+			<Include Script="TestConfigs/DefaultTestConfig.xml"/>
+		</Do>
+		<ForEach Name="TargetPlatform" Values="$(TargetPlatforms)" Separator="+">
+			<Property Name="GeneratedTargetTestNodes" Value="$(GeneratedTargetTestNodes);$(GeneratedTargetTestNodes$(TargetPlatform))"/>
+		</ForEach>
+	</Do>
+
 	<!-- Create an aggregate with all the generated nodes -->
 	<Aggregate Name="BuildAndTest $(ProjectName)" Requires="$(GeneratedToolNodes);$(GeneratedCookedEditorNodes);$(GeneratedEditorTestNodes);$(GeneratedCookedEditorTestNodes);$(GeneratedTargetNodes);$(GeneratedTargetTestNodes)"/>
 	<Aggregate Name="OnlyTest $(ProjectName)" Requires="$(GeneratedEditorTestNodes);$(GeneratedCookedEditorTestNodes);$(GeneratedTargetTestNodes)"/>
+
 	
 </BuildGraph>
\ No newline at end of file
diff --git a/Engine/Build/Graph/Tasks/RunTestConfigs.xml b/Engine/Build/Graph/Tasks/RunTestConfigs.xml
new file mode 100644
--- /dev/null
+++ b/Engine/Build/Graph/Tasks/RunTestConfigs.xml
@@ -0,0 +1,67 @@
+<?xml version='1.0' ?>
+<BuildGraph xmlns="http://www.epicgames.com/BuildGraph" xmlns:xsi="http://www.w3.org/2001/XMLSchema-instance" xsi:schemaLocation="http://www.epicgames.com/BuildGraph ../../../../Engine/Build/Graph/Schema.xsd" >
+	<Property Name="TestNodeName" Value="" />
+	<!--Iterate over all available platforms-->
+	<ForEach Name="TargetPlatform" Values="$(TargetPlatforms)" Separator="+"> 
+		<ForEach Name="TestConfig" Values="$(AllTestConfigs)" Separator=";">
+			<Property Name="TestConfigTitle" Value = ""/>
+			<Regex Input="$(TestConfig)" Pattern="(.*)\|(.*)\|(.*)\|(.*)\|(.*)" Capture="TestConfigTargetPlatform;TestConfigName;ArtifactPostfix;AdditionalArguments;AgentType"/>
+			<Property Name="TestConfigTitle" Value = " $(TestConfigName)" If="'$(TestConfigName)' != ''"/>
+
+			<Property Name="TestConfigArgs" Value="-AutomatedPerfTest.TestConfigName=&quot;$(TestConfigName)&quot;" />
+			<Property Name="TestConfigArgs" Value="$(TestConfigArgs) -ArtifactPostfix=$(ArtifactPostfix)" If="'$(ArtifactPostfix)' != ''"  />
+			<Property Name="TestConfigArgs" Value="$(TestConfigArgs) -args=&quot;$(AdditionalArguments)&quot;" If="'$(AdditionalArguments)' != ''" />
+
+			<!--If this is the default test config i.e. "All" ensure we create nodes for it-->
+			<Property Name="TestConfigTargetPlatform" Value="$(TargetPlatform)" If="'$(TestConfigTargetPlatform)' == 'All'"/>
+			<!-- <Property Name="TestConfigTargetPlatforms" Value="$(TargetPlatforms)" If="'$(TestConfigTargetPlatform)' == 'All'" /> -->
+
+				<Do If="'$(AgentType)' == 'Default'">
+					<Property Name="IsDesktopPlatform" Value="ContainsItem('$(AllDesktopPlatforms)', '$(TargetPlatform)', '+')"/>
+					<Property Name="IsTargetPlatformMac" Value="('$(TargetPlatform)' == 'Mac' or '$(TargetPlatform)' == 'IOS' or '$(TargetPlatform)' == 'tvOS')"/>
+					<Property Name="AgentType" Value="DevkitAutomation" />
+					<Property Name="AgentType" Value="Test$(TargetPlatform)" If="$(IsDesktopPlatform)"/>
+					<Property Name="AgentType" Value="$(TargetPlatform)" If="$(IsDesktopPlatform) and $(DontUseTestAgent)" />
+				</Do>
+
+				<!--Check if target platform applies to current test config-->
+				<Do If="'$(TestConfigTargetPlatform)' == '$(TargetPlatform)'">
+					<!-- Create the -build argument for Gauntlet based on whether builds are local or on the network-->
+					<Property Name="GauntletBuildArg" Value="$(GauntletBuildArg$(TargetPlatform))" />
+					<Property Name="GauntletTestArgs" Value="$(GauntletTargetTestArgs)"/>
+					<Property Name="GauntletTestArgs" Value="$(GauntletTargetTestAltArgs)" If="ContainsItem('$(AlternateDevicePoolPlatforms)', '$(TargetPlatform)', ';')"/>
+					<Property Name="GauntletTestArgs" Value="$(GauntletTestArgs) -adbkeys=&quot;$(AdbKeysPath)&quot;" If="'$(AdbKeysPath)' != '' and '$(TargetPlatform)' == 'Android'"/>
+					<ForEach Name="TestConfiguration" Values="$(TestConfigurations)" Separator="+">
+						<!-- Accumulate the new Test Config nodes per platform -->
+						<ForEach Name="TestName" Values="$(TargetTestList)" Separator="+">
+							<Property Name="NodeAnnotations" Value=""/>
+							<Property Name="NodeAnnotations" Value="Workflow=$(AnnotationsTarget)" If="ContainsItem('$(AnnotationsPlatforms)', '$(TargetPlatform)', ';') and ContainsItem('$(AnnotationsTests)', '$(TestName)', ';')"/>
+							<Agent Name="Test $(ProjectName) $(TargetPlatform)" Type="$(AgentType)">
+								<Property Name="TestNameTitle" Value="$(TestName)"/>
+								<Expand Name="SanitizeTestNameTitle"/>
+								<!-- Run and report on the test-->
+								<Do If="'$(AutomationReportOutputDirectory)' != ''">
+									<Property Name="OutputFragmentPath" Value="$(TargetPlatform)/$(TestConfiguration)/$(TestNameTitle)" />
+									<Property Name="OutputReportPath" Value="$(AutomationReportOutputDirectory)/$(OutputFragmentPath)" />
+									<Property Name="GauntletReportArgs" Value="-logdir=&quot;$(OutputReportPath)&quot; -artifactname=Gauntlet -ReportExportPath=&quot;$(OutputReportPath)&quot;" />
+									<Property Name="GauntletReportArgs" Value="$(GauntletReportArgs) -ReportURL=&quot;$(AutomationReportUrl)/$(OutputFragmentPath)/&quot;" If="'$(AutomationReportUrl)' != ''" />
+								</Do>
+								<Property Name="TestNodeName" Value="$(PreNodeName)$(TargetPlatform) $(TestConfiguration) $(TestNameTitle)$(TestConfigTitle)" />
+								<!-- Run the test. No 'Requires' as we run on the same machine locally and use the publish folder otherwise -->
+								<Property Name="TestNameBeforeParentheses" Value="$(TestName)"/>
+								<Property Name="TestArgumentsWithinParentheses" Value=""/>
+								<Expand Name="SplitCombinedTestNameAndArgs"/>
+								<Node Name="$(TestNodeName)" Requires="$(TestPrerequisites$(TargetPlatform)$(TestConfiguration))" Annotations="$(NodeAnnotations)">
+									<Command Name="RunUnreal" Arguments="-project=$(TargetProject) -test=&quot;$(TestNameBeforeParentheses)&quot; -platform=$(TargetPlatform) -configuration=$(TestConfiguration) -target=$(TargetName) $(GauntletBuildArg) $(GauntletTestArgs) $(TestArgumentsWithinParentheses) $(TestConfigArgs) $(GauntletReportArgs)"/>
+								</Node>
+							</Agent>
+							<!-- Add it to our list -->
+							<Do If="!ContainsItem('$(TargetPlatformsToSkipTests)', '$(TargetPlatform)', '+')">
+								<Property Name="GeneratedTargetTestNodes$(TargetPlatform)" Value="$(GeneratedTargetTestNodes$(TargetPlatform));$(TestNodeName)" />
+							</Do>
+						</ForEach>
+					</ForEach>
+				</Do>
+		</ForEach>
+	</ForEach>
+</BuildGraph>
\ No newline at end of file
diff --git a/Engine/Build/Graph/Tasks/TestConfigs.xml b/Engine/Build/Graph/Tasks/TestConfigs.xml
new file mode 100644
--- /dev/null
+++ b/Engine/Build/Graph/Tasks/TestConfigs.xml
@@ -0,0 +1,10 @@
+<?xml version='1.0' ?>
+<BuildGraph xmlns="http://www.epicgames.com/BuildGraph" xmlns:xsi="http://www.w3.org/2001/XMLSchema-instance" xsi:schemaLocation="http://www.epicgames.com/BuildGraph ../../../../Engine/Build/Graph/Schema.xsd" >
+
+	<Property Name="AllTestConfigs" Value=""/>
+	
+	<Macro Name="AddTestConfig" Arguments="TestPlatform;TestConfigName;ArtifactPostfix;AdditionalArguments;AgentType">
+		<Property Name="AllTestConfigs" Value="$(AllTestConfigs);$(TestPlatform)|$(TestConfigName)|$(ArtifactPostfix)|$(AdditionalArguments)|$(AgentType);"/>
+	</Macro>
+
+</BuildGraph>
diff --git a/Engine/Build/Graph/Tasks/TestConfigs/DefaultTestConfig.xml b/Engine/Build/Graph/Tasks/TestConfigs/DefaultTestConfig.xml
new file mode 100644
--- /dev/null
+++ b/Engine/Build/Graph/Tasks/TestConfigs/DefaultTestConfig.xml
@@ -0,0 +1,12 @@
+<?xml version='1.0' ?>
+<BuildGraph xmlns="http://www.epicgames.com/BuildGraph" xmlns:xsi="http://www.w3.org/2001/XMLSchema-instance" xsi:schemaLocation="http://www.epicgames.com/BuildGraph ../../../../Engine/Build/Graph/Schema.xsd" >
+
+	<Expand Name="AddTestConfig" 
+		TestPlatform="All"
+		TestConfigName=""
+		ArtifactPostfix=""
+		AdditionalArguments=""
+		AgentType="Default"
+	/>
+
+</BuildGraph>
\ No newline at end of file
diff --git a/Engine/Plugins/Performance/AutomatedPerfTesting/Build/Inc/AutomatedPerfTestCommonSettings.xml b/Engine/Plugins/Performance/AutomatedPerfTesting/Build/Inc/AutomatedPerfTestCommonSettings.xml
--- a/Engine/Plugins/Performance/AutomatedPerfTesting/Build/Inc/AutomatedPerfTestCommonSettings.xml
+++ b/Engine/Plugins/Performance/AutomatedPerfTesting/Build/Inc/AutomatedPerfTestCommonSettings.xml
@@ -6,9 +6,6 @@
 	<!-- insights traces after the conclusion of the test --> 
 	<Option Name="UseShippingInsights" Restrict="true|false" DefaultValue="false" Description="Whether or not to use the shipping configuration of Insights." />
 
-	<!-- Test configurations are more expressive way of determining of what variations of things need to be tested -->
-	<Include Script="TestConfigs.xml" />
-
 	<!-- then we'll use that to override values that will be picked up by the rest of BuildAndTestProject.xml -->
 	<Property Name="IncludeUnrealInsightsInToolCompile" Value="true" />
 	<Property Name="UnrealInsightsBuildConfiguration" Value="Development" If="!$(UseShippingInsights)" />
@@ -41,7 +38,6 @@
 	<Option Name="WithAPTStaticCameraTests" DefaultValue="false" Restrict="true|false" Description="Whether to run Auto Perf Test static camera tests."  />
 	<Option Name="WithAPTMaterialTest" DefaultValue="false" Restrict="true|false" Description="Whether to run Auto Perf Test material test."  />
 	<Option Name="WithAPTReplayTest"   DefaultValue="false" Restrict="true|false" Description="Whether to run Auto Perf Test replay test."  />
-	<Option Name="WithTestConfigs" DefaultValue="false" Restrict="true|false" Description="Whether or not to use the additional test configs" />
 
 	<!-- Test sub-types, act as flavors to the above tests -->
 	<Option Name="DoPerf" DefaultValue="false" Restrict="true|false" Description="Whether to run Perf (general performance) test." />
@@ -74,66 +70,37 @@
 		<Property Name="InsightsTestCommonArgs" Value="$(InsightsTestCommonArgs) -AutomatedPerfTest.UsePlatformExtraDebugMemory"/>
 	</Do>
 
-	<Property Name="AllTestConfigs" Value="" If="!$(WithTestConfigs)" />
-
-	<!-- Add in the base test config -->
-	<Property Name="AllTestConfigs" Value=".|.|.|.;$(AllTestConfigs)" />
-
-	<ForEach Name="TestConfig" Values="$(AllTestConfigs)" Separator=";">
-		<Regex Input="$(TestConfig)" Pattern="(.*)\|(.*)\|(.*)\|(.*)" Capture="PlatformToTest;TestConfigName;ArtifactPostfix;AdditionalArguments"/>
-
-		<Do If="'$(TestConfigName)' == '.'">
-			<Property Name="TestConfigName" Value="" />
-		</Do>
-		<Property Name="TestConfigArgs" Value=" -AutomatedPerfTest.TestConfigName=&quot;$(TestConfigName)&quot;" />
-
-		<Do If="'$(TestConfigName)' != ''">
-			<Property Name="TestConfigName" Value=" $(TestConfigName)" />
-		</Do>
-
-		<Do If="'$(ArtifactPostfix)' != '.'">
-			<Property Name="TestConfigArgs" Value="$(TestConfigArgs) -ArtifactPostfix=$(ArtifactPostfix)" />
-		</Do>
-
-		<Do If="'$(AdditionalArguments)' != '.'">
-			<Property Name="TestConfigArgs" Value="$(TestConfigArgs) -args=&quot;$(AdditionalArguments)&quot;" />
-		</Do>
-
-		<!-- The special sentinel value of '.' denotes to always run the base test, this will change once we refactor to always use TestConfigs -->
-		<Do If="ContainsItem('$(PlatformToTest)', '$(TargetPlatforms)', '+') or '$(PlatformToTest)' == '.'">
-			<!-- append plugin-provided tests to GauntletSettings' TargetTestList -->
-			<Do If="$(WithAPTSequenceTests)">
-				<ForEach Name="MapSequenceComboName" Values="$(MapSequenceComboNames)" Separator="+">
-					<Property Name="TargetTestList" Value="$(TargetTestList)AutomatedPerfTest.SequenceTest{-AutomatedPerfTest.SequencePerfTest.MapSequenceName=$(MapSequenceComboName) -AutomatedPerfTest.DoCSVProfiler $(PerfTestCommonArgs)$(TestConfigArgs)}[Sequence Perf $(MapSequenceComboName)$(TestConfigName)]+" If="$(DoPerf)" />
-					<Property Name="TargetTestList" Value="$(TargetTestList)AutomatedPerfTest.SequenceTest{-AutomatedPerfTest.SequencePerfTest.MapSequenceName=$(MapSequenceComboName) -AutomatedPerfTest.DoCSVProfiler -AutomatedPerfTest.DoLLM $(LLMTestCommonArgs)$(TestConfigArgs)}[Sequence LLM $(MapSequenceComboName)$(TestConfigName)]+" If="$(DoLLM)"/>
-					<Property Name="TargetTestList" Value="$(TargetTestList)AutomatedPerfTest.SequenceTest{-AutomatedPerfTest.SequencePerfTest.MapSequenceName=$(MapSequenceComboName) -AutomatedPerfTest.DoCSVProfiler -AutomatedPerfTest.DoGPUPerf $(GPUPerfTestCommonArgs)$(TestConfigArgs)}[Sequence GPUPerf $(MapSequenceComboName)$(TestConfigName)]+" If="$(DoGPUPerf)" />
-					<Property Name="TargetTestList" Value="$(TargetTestList)AutomatedPerfTest.SequenceTest{-AutomatedPerfTest.SequencePerfTest.MapSequenceName=$(MapSequenceComboName) -AutomatedPerfTest.DoCSVProfiler -AutomatedPerfTest.DoInsightsTrace $(InsightsTestCommonArgs)$(TestConfigArgs)}[Sequence Insights $(MapSequenceComboName)$(TestConfigName)]+" If="$(DoInsightsTrace)"/>
-				</ForEach>
-			</Do>
+	<!-- append plugin-provided tests to GauntletSettings' TargetTestList -->
+	<Do If="$(WithAPTSequenceTests)">
+		<ForEach Name="MapSequenceComboName" Values="$(MapSequenceComboNames)" Separator="+">
+			<Property Name="TargetTestList" Value="$(TargetTestList)AutomatedPerfTest.SequenceTest{-AutomatedPerfTest.SequencePerfTest.MapSequenceName=$(MapSequenceComboName) -AutomatedPerfTest.DoCSVProfiler $(PerfTestCommonArgs)}[Sequence Perf $(MapSequenceComboName)]+" If="$(DoPerf)" />
+			<Property Name="TargetTestList" Value="$(TargetTestList)AutomatedPerfTest.SequenceTest{-AutomatedPerfTest.SequencePerfTest.MapSequenceName=$(MapSequenceComboName) -AutomatedPerfTest.DoCSVProfiler -AutomatedPerfTest.DoLLM $(LLMTestCommonArgs)}[Sequence LLM $(MapSequenceComboName)]+" If="$(DoLLM)"/>
+			<Property Name="TargetTestList" Value="$(TargetTestList)AutomatedPerfTest.SequenceTest{-AutomatedPerfTest.SequencePerfTest.MapSequenceName=$(MapSequenceComboName) -AutomatedPerfTest.DoCSVProfiler -AutomatedPerfTest.DoGPUPerf $(GPUPerfTestCommonArgs)}[Sequence GPUPerf $(MapSequenceComboName)]+" If="$(DoGPUPerf)" />
+			<Property Name="TargetTestList" Value="$(TargetTestList)AutomatedPerfTest.SequenceTest{-AutomatedPerfTest.SequencePerfTest.MapSequenceName=$(MapSequenceComboName) -AutomatedPerfTest.DoCSVProfiler -AutomatedPerfTest.DoInsightsTrace $(InsightsTestCommonArgs)}[Sequence Insights $(MapSequenceComboName)]+" If="$(DoInsightsTrace)"/>
+		</ForEach>
+	</Do>
 
-			<Do If="$(WithAPTStaticCameraTests)">
-				<ForEach Name="StaticCameraMapName" Values="$(StaticCameraMapNames)" Separator="+">
-					<Property Name="TargetTestList" Value="$(TargetTestList)AutomatedPerfTest.StaticCameraTest{-AutomatedPerfTest.StaticCameraPerfTest.MapName=$(StaticCameraMapName) -AutomatedPerfTest.DoCSVProfiler $(PerfTestCommonArgs)$(TestConfigArgs)}[Static Camera Perf $(StaticCameraMapName)$(TestConfigName)]+" If="$(DoPerf)" />
-					<Property Name="TargetTestList" Value="$(TargetTestList)AutomatedPerfTest.StaticCameraTest{-AutomatedPerfTest.StaticCameraPerfTest.MapName=$(StaticCameraMapName) -AutomatedPerfTest.DoCSVProfiler -AutomatedPerfTest.DoLLM $(LLMTestCommonArgs)$(TestConfigArgs)}[Static Camera LLM $(StaticCameraMapName)$(TestConfigName)]+" If="$(DoLLM)"/>
-					<Property Name="TargetTestList" Value="$(TargetTestList)AutomatedPerfTest.StaticCameraTest{-AutomatedPerfTest.StaticCameraPerfTest.MapName=$(StaticCameraMapName) -AutomatedPerfTest.DoCSVProfiler -AutomatedPerfTest.DoGPUPerf $(LLMTestCommonArgs)$(TestConfigArgs)}[Static Camera GPUPerf $(StaticCameraMapName)$(TestConfigName)]+" If="$(DoGPUPerf)"/>
-					<Property Name="TargetTestList" Value="$(TargetTestList)AutomatedPerfTest.StaticCameraTest{-AutomatedPerfTest.StaticCameraPerfTest.MapName=$(StaticCameraMapName) -AutomatedPerfTest.DoCSVProfiler -AutomatedPerfTest.DoInsightsTrace $(InsightsTestCommonArgs)$(TestConfigArgs)}[Static Camera Insights $(StaticCameraMapName)$(TestConfigName)]+" If="$(DoInsightsTrace)"/>
-				</ForEach>
-			</Do>
+	<Do If="$(WithAPTStaticCameraTests)">
+		<ForEach Name="StaticCameraMapName" Values="$(StaticCameraMapNames)" Separator="+">
+			<Property Name="TargetTestList" Value="$(TargetTestList)AutomatedPerfTest.StaticCameraTest{-AutomatedPerfTest.StaticCameraPerfTest.MapName=$(StaticCameraMapName) -AutomatedPerfTest.DoCSVProfiler $(PerfTestCommonArgs)}[Static Camera Perf $(StaticCameraMapName)]+" If="$(DoPerf)" />
+			<Property Name="TargetTestList" Value="$(TargetTestList)AutomatedPerfTest.StaticCameraTest{-AutomatedPerfTest.StaticCameraPerfTest.MapName=$(StaticCameraMapName) -AutomatedPerfTest.DoCSVProfiler -AutomatedPerfTest.DoLLM $(LLMTestCommonArgs)}[Static Camera LLM $(StaticCameraMapName)]+" If="$(DoLLM)"/>
+			<Property Name="TargetTestList" Value="$(TargetTestList)AutomatedPerfTest.StaticCameraTest{-AutomatedPerfTest.StaticCameraPerfTest.MapName=$(StaticCameraMapName) -AutomatedPerfTest.DoCSVProfiler -AutomatedPerfTest.DoGPUPerf $(LLMTestCommonArgs)}[Static Camera GPUPerf $(StaticCameraMapName)]+" If="$(DoGPUPerf)"/>
+			<Property Name="TargetTestList" Value="$(TargetTestList)AutomatedPerfTest.StaticCameraTest{-AutomatedPerfTest.StaticCameraPerfTest.MapName=$(StaticCameraMapName) -AutomatedPerfTest.DoCSVProfiler -AutomatedPerfTest.DoInsightsTrace $(InsightsTestCommonArgs)}[Static Camera Insights $(StaticCameraMapName)]+" If="$(DoInsightsTrace)"/>
+		</ForEach>
+	</Do>
 
-			<Do If="$(WithAPTMaterialTest)">
-				<Property Name="TargetTestList" Value="$(TargetTestList)AutomatedPerfTest.MaterialTest{-AutomatedPerfTest.DoCSVProfiler $(PerfTestCommonArgs)$(TestConfigArgs)}[Material Perf$(TestConfigName)]+" If="$(DoPerf)" />
-				<Property Name="TargetTestList" Value="$(TargetTestList)AutomatedPerfTest.MaterialTest{-AutomatedPerfTest.DoCSVProfiler -AutomatedPerfTest.DoLLM $(LLMTestCommonArgs)$(TestConfigArgs)}[Material LLM$(TestConfigName)]+" If="$(DoLLM)"/>
-				<Property Name="TargetTestList" Value="$(TargetTestList)AutomatedPerfTest.MaterialTest{-AutomatedPerfTest.DoCSVProfiler -AutomatedPerfTest.DoGPUPerf $(LLMTestCommonArgs)$(TestConfigArgs)}[Material GPUPerf$(TestConfigName)]+" If="$(DoGPUPerf)"/>
-				<Property Name="TargetTestList" Value="$(TargetTestList)AutomatedPerfTest.MaterialTest{-AutomatedPerfTest.DoCSVProfiler -AutomatedPerfTest.DoInsightsTrace $(InsightsTestCommonArgs)$(TestConfigArgs)}[Material Insights$(TestConfigName)]+" If="$(DoInsightsTrace)"/>
-			</Do>
+	<Do If="$(WithAPTMaterialTest)">
+		<Property Name="TargetTestList" Value="$(TargetTestList)AutomatedPerfTest.MaterialTest{-AutomatedPerfTest.DoCSVProfiler $(PerfTestCommonArgs)}[Material Perf]+" If="$(DoPerf)" />
+		<Property Name="TargetTestList" Value="$(TargetTestList)AutomatedPerfTest.MaterialTest{-AutomatedPerfTest.DoCSVProfiler -AutomatedPerfTest.DoLLM $(LLMTestCommonArgs)}[Material LLM]+" If="$(DoLLM)"/>
+		<Property Name="TargetTestList" Value="$(TargetTestList)AutomatedPerfTest.MaterialTest{-AutomatedPerfTest.DoCSVProfiler -AutomatedPerfTest.DoGPUPerf $(LLMTestCommonArgs)}[Material GPUPerf]+" If="$(DoGPUPerf)"/>
+		<Property Name="TargetTestList" Value="$(TargetTestList)AutomatedPerfTest.MaterialTest{-AutomatedPerfTest.DoCSVProfiler -AutomatedPerfTest.DoInsightsTrace $(InsightsTestCommonArgs)}[Material Insights]+" If="$(DoInsightsTrace)"/>
+	</Do>
 
-			<Do If="$(WithAPTReplayTest)">
-				<Property Name="TargetTestList" Value="$(TargetTestList)AutomatedPerfTest.ReplayTest{-AutomatedPerfTest.ReplayPerfTest.ReplayName=&quot;$(ReplayName)&quot; -AutomatedPerfTest.DoCSVProfiler $(PerfTestCommonArgs)$(TestConfigArgs)}[Replay Perf$(TestConfigName)]+" If="$(DoPerf)" />
-				<Property Name="TargetTestList" Value="$(TargetTestList)AutomatedPerfTest.ReplayTest{-AutomatedPerfTest.ReplayPerfTest.ReplayName=&quot;$(ReplayName)&quot; -AutomatedPerfTest.DoCSVProfiler -AutomatedPerfTest.DoLLM $(LLMTestCommonArgs)$(TestConfigArgs)}[Replay LLM$(TestConfigName)]+" If="$(DoLLM)"/>
-				<Property Name="TargetTestList" Value="$(TargetTestList)AutomatedPerfTest.ReplayTest{-AutomatedPerfTest.ReplayPerfTest.ReplayName=&quot;$(ReplayName)&quot; -AutomatedPerfTest.DoCSVProfiler -AutomatedPerfTest.DoGPUPerf $(LLMTestCommonArgs)$(TestConfigArgs)}[Replay GPUPerf$(TestConfigName)]+" If="$(DoGPUPerf)"/>
-				<Property Name="TargetTestList" Value="$(TargetTestList)AutomatedPerfTest.ReplayTest{-AutomatedPerfTest.ReplayPerfTest.ReplayName=&quot;$(ReplayName)&quot; -AutomatedPerfTest.DoCSVProfiler -AutomatedPerfTest.DoInsightsTrace $(InsightsTestCommonArgs)$(TestConfigArgs)}[Replay Insights$(TestConfigName)]+" If="$(DoInsightsTrace)"/>
-			</Do>
-		</Do>
-	</ForEach>
+	<Do If="$(WithAPTReplayTest)">
+		<Property Name="TargetTestList" Value="$(TargetTestList)AutomatedPerfTest.ReplayTest{-AutomatedPerfTest.ReplayPerfTest.ReplayName=&quot;$(ReplayName)&quot; -AutomatedPerfTest.DoCSVProfiler $(PerfTestCommonArgs)}[Replay Perf]+" If="$(DoPerf)" />
+		<Property Name="TargetTestList" Value="$(TargetTestList)AutomatedPerfTest.ReplayTest{-AutomatedPerfTest.ReplayPerfTest.ReplayName=&quot;$(ReplayName)&quot; -AutomatedPerfTest.DoCSVProfiler -AutomatedPerfTest.DoLLM $(LLMTestCommonArgs)}[Replay LLM]+" If="$(DoLLM)"/>
+		<Property Name="TargetTestList" Value="$(TargetTestList)AutomatedPerfTest.ReplayTest{-AutomatedPerfTest.ReplayPerfTest.ReplayName=&quot;$(ReplayName)&quot; -AutomatedPerfTest.DoCSVProfiler -AutomatedPerfTest.DoGPUPerf $(LLMTestCommonArgs)}[Replay GPUPerf]+" If="$(DoGPUPerf)"/>
+		<Property Name="TargetTestList" Value="$(TargetTestList)AutomatedPerfTest.ReplayTest{-AutomatedPerfTest.ReplayPerfTest.ReplayName=&quot;$(ReplayName)&quot; -AutomatedPerfTest.DoCSVProfiler -AutomatedPerfTest.DoInsightsTrace $(InsightsTestCommonArgs)}[Replay Insights]+" If="$(DoInsightsTrace)"/>
+	</Do>
 
 </BuildGraph>
diff --git a/Engine/Plugins/Performance/AutomatedPerfTesting/Build/Inc/TestConfigs.xml b/Engine/Plugins/Performance/AutomatedPerfTesting/Build/Inc/TestConfigs.xml
--- a/Engine/Plugins/Performance/AutomatedPerfTesting/Build/Inc/TestConfigs.xml
+++ b/Engine/Plugins/Performance/AutomatedPerfTesting/Build/Inc/TestConfigs.xml
@@ -3,36 +3,8 @@
 
 	<Property Name="AllTestConfigs" Value=""/>
 	
-	<Macro Name="AddTestConfig" Arguments="TestPlatform;TestConfigName;ArtifactPostfix;AdditionalArguments">
-		<Property Name="AllTestConfigs" Value="$(AllTestConfigs);$(TestPlatform)|$(TestConfigName)|$(ArtifactPostfix)|$(AdditionalArguments);"/>
+	<Macro Name="AddTestConfig" Arguments="TestPlatform;TestConfigName;ArtifactPostfix;AdditionalArguments;AgentType">
+		<Property Name="AllTestConfigs" Value="$(AllTestConfigs);$(TestPlatform)|$(TestConfigName)|$(ArtifactPostfix)|$(AdditionalArguments)"/>
 	</Macro>
 
-	<Expand Name="AddTestConfig" 
-		TestPlatform="Win64"
-		TestConfigName="Vulkan SM6"
-		ArtifactPostfix="vulkan_sm6"
-		AdditionalArguments="-vulkan -sm6"
-	/>
-
-	<Expand Name="AddTestConfig" 
-		TestPlatform="Linux"
-		TestConfigName="Vulkan SM5"
-		ArtifactPostfix="vulkan_sm5"
-		AdditionalArguments="-vulkan -sm5"
-	/>
-
-	<Expand Name="AddTestConfig" 
-		TestPlatform="Win64"
-		TestConfigName="Vulkan SM5"
-		ArtifactPostfix="vulkan_sm5"
-		AdditionalArguments="-vulkan -sm5"
-	/>
-
-	<Expand Name="AddTestConfig" 
-		TestPlatform="Linux"
-		TestConfigName="Vulkan ES31"
-		ArtifactPostfix="vulkan_es31"
-		AdditionalArguments="-vulkan -es31"
-	/>
-
 </BuildGraph>
