Added GPU-Reshape automated test
- Added desktop build bootstrapping, simpler compared to custom module roles

diff --git a/Engine/Build/Graph/Tasks/BuildAndTestProject.xml b/Engine/Build/Graph/Tasks/BuildAndTestProject.xml
--- a/Engine/Build/Graph/Tasks/BuildAndTestProject.xml
+++ b/Engine/Build/Graph/Tasks/BuildAndTestProject.xml
@@ -60,6 +60,7 @@
 	<Option Name="IncludeUnrealInsightsInToolCompile" DefaultValue="false" Description="Include Unreal Insights into Compile Tools prerequisites" />
 	<Option Name="IncludeEpicWebHelperInToolCompile" DefaultValue="false" Description="Include EpicWebHelper into Compile Tools prerequisites" />
 	<Option Name="WithDeviceReservation" DefaultValue="$(Horde)" Description="Whether to re-use devices and installed builds for a sequence of tests, where supported"/>
+	<Option Name="DoGPUReshape" DefaultValue="false" Restrict="true|false" Description="Whether to run GPU instrumentation test with GPU Reshape." />
 
 	<!-- 
 		Allow project graphs that Include this file to specify their own defaults.
@@ -216,6 +217,12 @@
 	<Option Name="RetrieveUsingCloudArtifacts" DefaultValue="false" Restrict="true|false" Description="Enable to retrieve cloud artifacts" />
 	<Option Name="PublishCloudArtifacts" DefaultValue="$(RetrieveUsingCloudArtifacts)" Description="Enable to publish cloud artifacts" />
 
+	<!-- If running with GPU Reshape, override the local build configuration
+	     Enables shader symbols and anything else that aids troubleshooting -->
+	<Do If="$(DoGPUReshape)">
+		<Property Name="ExtraCookArguments" Value="$(ExtraCookArguments) -cvarsini=$(RootDir)\Engine\Source\Programs\GPUReshape\Resources\BuildConfig.ini"/>
+	</Do>
+
 	<!-- Detect the type of the client based on the target name -->
 	<Property Name="TargetType" Value="Game" />
 	<Property Name="TargetType" Value="Client" If="Contains('$(TargetName)', 'Client')" />
@@ -284,6 +291,7 @@
 			<Node Name="$(PreNodeName)Compile Tools $(EditorPlatform)" Requires="$(PreNodeName)Update Version Files $(EditorPlatform)">
 				<Compile Target="ShaderCompileWorker" Platform="$(EditorPlatform)" Configuration="Development" Arguments="$(GenericCompileArguments)"/>
 				<Compile Target="ZenLaunch" Platform="$(EditorPlatform)" Configuration="Development" Arguments="$(GenericCompileArguments)"/>
+				<Compile Target="GPUReshapeBootstrapper" Platform="$(EditorPlatform)" Configuration="Development" Arguments="$(GenericCompileArguments)" If="$(DoGPUReshape)"/>
 				<Compile Target="UnrealPak" Platform="$(EditorPlatform)" Configuration="Development" Arguments="$(GenericCompileArguments)"/>
 				<Compile Target="CrashReportClient" Platform="$(EditorPlatform)" Configuration="Shipping" Arguments="$(GenericCompileArguments)"/>
 				<Compile Target="CrashReportClientEditor" Platform="$(EditorPlatform)" Configuration="Shipping" Arguments="$(GenericCompileArguments)"/>
diff --git a/Engine/Plugins/Performance/AutomatedPerfTesting/Build/Inc/AutomatedPerfTestCommonSettings.xml b/Engine/Plugins/Performance/AutomatedPerfTesting/Build/Inc/AutomatedPerfTestCommonSettings.xml
--- a/Engine/Plugins/Performance/AutomatedPerfTesting/Build/Inc/AutomatedPerfTestCommonSettings.xml
+++ b/Engine/Plugins/Performance/AutomatedPerfTesting/Build/Inc/AutomatedPerfTestCommonSettings.xml
@@ -27,9 +27,11 @@
 	<Option Name="IterationsLLM" DefaultValue="3" Restrict="^[0-9]+$" Description="Number of iterations to run LLM test with" />
 	<Option Name="IterationsGPUPerf" DefaultValue="3" Restrict="^[0-9]+$" Description="Number of iterations to run GPUPerf test with" />
 	<Option Name="IterationsInsights" DefaultValue="1" Restrict="^[0-9]+$" Description="Number of iterations to run Insights test with" />
+	<Option Name="IterationsGPUReshape" DefaultValue="1" Restrict="^[0-9]+$" Description="Number of iterations to run GPUReshape test with" />
 	<Option Name="UsePlatformExtraDebugMemory" DefaultValue="true" Restrict="true|false" Description="Use additional debug memory on target device if supported. Applies only to LLM and Insights trace."/>
 	<Option Name="APTIgnoreTestBuildLogging" Restrict="true|false" DefaultValue="true" Description="If true, tests run in Test build configuration will not fail if logs are empty." />
 	<Option Name="PerfOutputPath" DefaultValue="" Description="Path where Perf Test artifacts (Logs, CSVs, Reports, Cache) will be stored. Override to a common network-accessible path when running on build machines if needed"/>
+	<Option Name="ProjectPath" DefaultValue="$(DefaultProjectPath)" Restrict=".+" Description="Path to the folder that contains your project"/>
 	
 	<Option Name="ExtraAutomatedPerformanceCommonArgs" DefaultValue="" Description="Additional Common arguments passed to Gauntlet command line"/>
 
@@ -44,6 +46,7 @@
 	<Option Name="DoLLM" DefaultValue="false" Restrict="true|false" Description="Whether to run LLM (memory) test." />
 	<Option Name="DoGPUPerf" DefaultValue="false" Restrict="true|false" Description="Whether to run GPUPerf (GPU performance) test." />
 	<Option Name="DoInsightsTrace" DefaultValue="false" Restrict="true|false" Description="Whether or not to capture Unreal Insights trace during the run." />
+	<Option Name="DoGPUReshape" DefaultValue="false" Restrict="true|false" Description="Whether to run GPU instrumentation test with GPU Reshape." />
 
 	<!-- Construct the common arguments for the performance test and report commands -->
 	<Property Name="AutomatedPerformanceCommonArgs" Value = "-AutomatedPerfTest.UseShippingInsights=$(UseShippingInsights)" />
@@ -64,7 +67,21 @@
 	<Property Name="LLMTestCommonArgs"      Value="$(AutomatedPerformanceCommonArgs) -iterations=$(IterationsLLM)"/>
 	<Property Name="GPUPerfTestCommonArgs"      Value="$(AutomatedPerformanceCommonArgs) -iterations=$(IterationsGPUPerf)"/>
 	<Property Name="InsightsTestCommonArgs" Value="$(AutomatedPerformanceCommonArgs) -iterations=$(IterationsInsights)"/>
-
+	
+	<!-- Reshape configuration -->
+	<Do If="$(DoGPUReshape)">
+		<Property Name="GPUReshapePath"           Value="$(RootDir)\Engine\Restricted\NotForLicensees\Binaries\ThirdParty\GPUReshape\Win64"/>
+		<Property Name="GPUReshapeBootstrapPath"  Value="$(RootDir)\Engine\Source\Programs\GPUReshape"/>
+		<Property Name="GPUReshapeReportPath"     Value="GRS.Report.json"/>
+		<Property Name="GPUReshapeReportPath"     Value="$(PerfOutputPath)\$(GPUReshapeReportPath)" If="'$(PerfOutputPath)' != ''"/>
+		<Property Name="GPUReshapeWorkspacePath"  Value="$(GPUReshapeBootstrapPath)\Resources\Workspaces\BasicWorkspace.json"/>
+		<Property Name="GPUReshapeSymbolPath"     Value="$(RootDir)\$(ProjectPath)\Saved\ShaderSymbols"/>
+		
+		<!-- Bootstrap arguments -->
+		<Property Name="GPUReshapeArgs"           Value="-GRS.Path=$(GPUReshapePath)\GPUReshape.exe -GRS.Workspace=$(GPUReshapeWorkspacePath) -GRS.Report=$(GPUReshapeReportPath) -GRS.SymbolPath=$(GPUReshapeSymbolPath) -nothreadtimeout -noheartbeatthread"/>
+		<Property Name="GPUReshapeTestCommonArgs" Value="$(AutomatedPerformanceCommonArgs) -Bootstrap=GPUReshapeBootstrapper.exe -args=&quot;$(GPUReshapeArgs)&quot; -iterations=$(IterationsGPUReshape)"/>
+	</Do>
+	
 	<Do If="$(UsePlatformExtraDebugMemory)">
 		<Property Name="LLMTestCommonArgs"      Value="$(LLMTestCommonArgs) -AutomatedPerfTest.UsePlatformExtraDebugMemory"/>
 		<Property Name="InsightsTestCommonArgs" Value="$(InsightsTestCommonArgs) -AutomatedPerfTest.UsePlatformExtraDebugMemory"/>
@@ -77,6 +94,7 @@
 			<Property Name="TargetTestList" Value="$(TargetTestList)AutomatedPerfTest.SequenceTest{-AutomatedPerfTest.SequencePerfTest.MapSequenceName=$(MapSequenceComboName) -AutomatedPerfTest.DoCSVProfiler -AutomatedPerfTest.DoLLM $(LLMTestCommonArgs)}[Sequence LLM $(MapSequenceComboName)]+" If="$(DoLLM)"/>
 			<Property Name="TargetTestList" Value="$(TargetTestList)AutomatedPerfTest.SequenceTest{-AutomatedPerfTest.SequencePerfTest.MapSequenceName=$(MapSequenceComboName) -AutomatedPerfTest.DoCSVProfiler -AutomatedPerfTest.DoGPUPerf $(GPUPerfTestCommonArgs)}[Sequence GPUPerf $(MapSequenceComboName)]+" If="$(DoGPUPerf)" />
 			<Property Name="TargetTestList" Value="$(TargetTestList)AutomatedPerfTest.SequenceTest{-AutomatedPerfTest.SequencePerfTest.MapSequenceName=$(MapSequenceComboName) -AutomatedPerfTest.DoCSVProfiler -AutomatedPerfTest.DoInsightsTrace $(InsightsTestCommonArgs)}[Sequence Insights $(MapSequenceComboName)]+" If="$(DoInsightsTrace)"/>
+			<Property Name="TargetTestList" Value="$(TargetTestList)AutomatedPerfTest.SequenceTest{-AutomatedPerfTest.SequencePerfTest.MapSequenceName=$(MapSequenceComboName) -AutomatedPerfTest.DoGPUReshape $(GPUReshapeTestCommonArgs)}[Sequence GPUReshape $(MapSequenceComboName)]+" If="$(DoGPUReshape)"/>
 		</ForEach>
 	</Do>
 
@@ -86,6 +104,7 @@
 			<Property Name="TargetTestList" Value="$(TargetTestList)AutomatedPerfTest.StaticCameraTest{-AutomatedPerfTest.StaticCameraPerfTest.MapName=$(StaticCameraMapName) -AutomatedPerfTest.DoCSVProfiler -AutomatedPerfTest.DoLLM $(LLMTestCommonArgs)}[Static Camera LLM $(StaticCameraMapName)]+" If="$(DoLLM)"/>
 			<Property Name="TargetTestList" Value="$(TargetTestList)AutomatedPerfTest.StaticCameraTest{-AutomatedPerfTest.StaticCameraPerfTest.MapName=$(StaticCameraMapName) -AutomatedPerfTest.DoCSVProfiler -AutomatedPerfTest.DoGPUPerf $(LLMTestCommonArgs)}[Static Camera GPUPerf $(StaticCameraMapName)]+" If="$(DoGPUPerf)"/>
 			<Property Name="TargetTestList" Value="$(TargetTestList)AutomatedPerfTest.StaticCameraTest{-AutomatedPerfTest.StaticCameraPerfTest.MapName=$(StaticCameraMapName) -AutomatedPerfTest.DoCSVProfiler -AutomatedPerfTest.DoInsightsTrace $(InsightsTestCommonArgs)}[Static Camera Insights $(StaticCameraMapName)]+" If="$(DoInsightsTrace)"/>
+			<Property Name="TargetTestList" Value="$(TargetTestList)AutomatedPerfTest.StaticCameraTest{-AutomatedPerfTest.StaticCameraPerfTest.MapName=$(StaticCameraMapName) -AutomatedPerfTest.DoGPUReshape $(GPUReshapeTestCommonArgs)}[Static Camera GPUReshape $(StaticCameraMapName)]+" If="$(DoGPUReshape)"/>
 		</ForEach>
 	</Do>
 
@@ -94,6 +113,7 @@
 		<Property Name="TargetTestList" Value="$(TargetTestList)AutomatedPerfTest.MaterialTest{-AutomatedPerfTest.DoCSVProfiler -AutomatedPerfTest.DoLLM $(LLMTestCommonArgs)}[Material LLM]+" If="$(DoLLM)"/>
 		<Property Name="TargetTestList" Value="$(TargetTestList)AutomatedPerfTest.MaterialTest{-AutomatedPerfTest.DoCSVProfiler -AutomatedPerfTest.DoGPUPerf $(LLMTestCommonArgs)}[Material GPUPerf]+" If="$(DoGPUPerf)"/>
 		<Property Name="TargetTestList" Value="$(TargetTestList)AutomatedPerfTest.MaterialTest{-AutomatedPerfTest.DoCSVProfiler -AutomatedPerfTest.DoInsightsTrace $(InsightsTestCommonArgs)}[Material Insights]+" If="$(DoInsightsTrace)"/>
+		<Property Name="TargetTestList" Value="$(TargetTestList)AutomatedPerfTest.MaterialTest{-AutomatedPerfTest.DoCSVProfiler -AutomatedPerfTest.DoGPUReshape $(GPUReshapeTestCommonArgs)}[Material GPUReshape]+" If="$(DoGPUReshape)"/>
 	</Do>
 
 	<Do If="$(WithAPTReplayTest)">
@@ -101,6 +121,7 @@
 		<Property Name="TargetTestList" Value="$(TargetTestList)AutomatedPerfTest.ReplayTest{-AutomatedPerfTest.ReplayPerfTest.ReplayName=&quot;$(ReplayName)&quot; -AutomatedPerfTest.DoCSVProfiler -AutomatedPerfTest.DoLLM $(LLMTestCommonArgs)}[Replay LLM]+" If="$(DoLLM)"/>
 		<Property Name="TargetTestList" Value="$(TargetTestList)AutomatedPerfTest.ReplayTest{-AutomatedPerfTest.ReplayPerfTest.ReplayName=&quot;$(ReplayName)&quot; -AutomatedPerfTest.DoCSVProfiler -AutomatedPerfTest.DoGPUPerf $(LLMTestCommonArgs)}[Replay GPUPerf]+" If="$(DoGPUPerf)"/>
 		<Property Name="TargetTestList" Value="$(TargetTestList)AutomatedPerfTest.ReplayTest{-AutomatedPerfTest.ReplayPerfTest.ReplayName=&quot;$(ReplayName)&quot; -AutomatedPerfTest.DoCSVProfiler -AutomatedPerfTest.DoInsightsTrace $(InsightsTestCommonArgs)}[Replay Insights]+" If="$(DoInsightsTrace)"/>
+		<Property Name="TargetTestList" Value="$(TargetTestList)AutomatedPerfTest.ReplayTest{-AutomatedPerfTest.ReplayPerfTest.ReplayName=&quot;$(ReplayName)&quot; -AutomatedPerfTest.DoGPUReshape $(GPUReshapeTestCommonArgs)}[Replay GPUReshape]+" If="$(DoGPUReshape)"/>
 	</Do>
 
 </BuildGraph>
diff --git a/Engine/Plugins/Performance/AutomatedPerfTesting/Build/Scripts/AutomatedPerfTestConfig.cs b/Engine/Plugins/Performance/AutomatedPerfTesting/Build/Scripts/AutomatedPerfTestConfig.cs
--- a/Engine/Plugins/Performance/AutomatedPerfTesting/Build/Scripts/AutomatedPerfTestConfig.cs
+++ b/Engine/Plugins/Performance/AutomatedPerfTesting/Build/Scripts/AutomatedPerfTestConfig.cs
@@ -76,6 +76,12 @@ namespace AutomatedPerfTest
 		[AutoParamWithNames(false, "AutomatedPerfTest.DoGPUPerf")]
 		public bool DoGPUPerf;
 
+		/// <summary>
+		/// Whether to run GPU Reshape's instrumentation on the target
+		/// </summary>
+		[AutoParamWithNames(false, "AutomatedPerfTest.DoGPUReshape")]
+		public bool DoGPUReshape;
+
 		/// <summary>
 		/// Whether to lock dynamic resolution. This is applied automatically when GPUPerf test is run.
 		/// </summary>
diff --git a/Engine/Plugins/Performance/AutomatedPerfTesting/Build/Scripts/AutomatedPerfTestNode.cs b/Engine/Plugins/Performance/AutomatedPerfTesting/Build/Scripts/AutomatedPerfTestNode.cs
--- a/Engine/Plugins/Performance/AutomatedPerfTesting/Build/Scripts/AutomatedPerfTestNode.cs
+++ b/Engine/Plugins/Performance/AutomatedPerfTesting/Build/Scripts/AutomatedPerfTestNode.cs
@@ -725,6 +725,11 @@ namespace AutomatedPerfTest
 				return "GPUPerf";
 			}
 
+			if (Config.DoGPUReshape)
+			{
+				return "GPUReshape";
+			}
+
 			return "Perf";
 		}
 
diff --git a/Engine/Source/Programs/AutomationTool/Gauntlet/Platform/Windows/Gauntlet.TargetDeviceWindows.cs b/Engine/Source/Programs/AutomationTool/Gauntlet/Platform/Windows/Gauntlet.TargetDeviceWindows.cs
--- a/Engine/Source/Programs/AutomationTool/Gauntlet/Platform/Windows/Gauntlet.TargetDeviceWindows.cs
+++ b/Engine/Source/Programs/AutomationTool/Gauntlet/Platform/Windows/Gauntlet.TargetDeviceWindows.cs
@@ -12,6 +12,7 @@ using System.Runtime.InteropServices;
 using System.Text.RegularExpressions;
 using System.Collections.Generic;
 using System.Linq;
+using UnrealBuildBase;
 
 namespace Gauntlet
 {
@@ -20,6 +21,12 @@ namespace Gauntlet
 	/// </summary>
 	public class TargetDeviceWindows : TargetDeviceDesktopCommon
 	{
+		/// <summary>
+		/// Optional executable that bootstraps the build
+		/// </summary>
+		[AutoParam]
+		public string Bootstrap { get; protected set; }
+		
 		public TargetDeviceWindows(string InName, string InCacheDir)
 			: base(InName, InCacheDir)
 		{
@@ -199,6 +206,29 @@ namespace Gauntlet
 			}
 		}
 
+		protected void BootstrapInstall(WindowsAppInstall Install)
+		{
+			if (string.IsNullOrEmpty(Bootstrap))
+			{
+				return;
+			}
+
+			string[] SearchPaths =
+			[
+				Path.GetDirectoryName(Install.ExecutablePath),
+				Path.Combine(Unreal.RootDirectory.FullName, "Engine", "Binaries", Platform.ToString())
+			];
+
+			if (SearchPaths.Select(x => Path.Combine(x, Bootstrap)).FirstOrDefault(File.Exists) is not { } BootstrapCandidate)
+			{
+				Log.Error("Failed to find bootstrapper '{0}'", Bootstrap);
+				return;
+			}
+			
+			Install.CommandArguments += " -BootstrapTarget=\"" + Install.ExecutablePath + "\"";
+			Install.ExecutablePath = BootstrapCandidate;
+		}
+
 		protected override IAppInstall InstallNativeStagedBuild(UnrealAppConfig AppConfig, NativeStagedBuild InBuild)
 		{
 			WindowsAppInstall WinApp = new WindowsAppInstall(AppConfig.Name, AppConfig.ProjectName, this);
@@ -206,6 +236,8 @@ namespace Gauntlet
 			WinApp.SetDefaultCommandLineArguments(AppConfig, RunOptions, InBuild.BuildPath);
 			WinApp.ExecutablePath = Path.Combine(InBuild.BuildPath, InBuild.ExecutablePath);
 
+			BootstrapInstall(WinApp);
+
 			CopyAdditionalFiles(AppConfig.FilesToCopy);
 
 			return WinApp;
@@ -275,6 +307,8 @@ namespace Gauntlet
 				WinApp.ExecutablePath = BinaryPath;
 			}
 
+			BootstrapInstall(WinApp);
+
 			return WinApp;
 		}
 
@@ -291,6 +325,8 @@ namespace Gauntlet
 				PopulateDirectoryMappings(AppConfig.ProjectFile.Directory.FullName);
 			}
 
+			BootstrapInstall(WinApp);
+
 			CopyAdditionalFiles(AppConfig.FilesToCopy);
 
 			return WinApp;
@@ -304,6 +340,8 @@ namespace Gauntlet
 
 			PopulateDirectoryMappings(Path.Combine(BasePath, AppConfig.ProjectName));
 
+			BootstrapInstall(WinApp);
+
 			CopyAdditionalFiles(AppConfig.FilesToCopy);
 			return WinApp;
 		}
diff --git a/Engine/Source/Programs/GPUReshape/Resources/BuildConfig.ini b/Engine/Source/Programs/GPUReshape/Resources/BuildConfig.ini
new file mode 100644
--- /dev/null
+++ b/Engine/Source/Programs/GPUReshape/Resources/BuildConfig.ini
@@ -0,0 +1,5 @@
+; Build configuration overrides
+
+[Startup]
+; Always enable embedded symbols, for source traceback
+r.Shaders.Symbols=1
diff --git a/Engine/Source/Programs/GPUReshape/Resources/Workspaces/BasicWorkspace.json b/Engine/Source/Programs/GPUReshape/Resources/Workspaces/BasicWorkspace.json
new file mode 100644
--- /dev/null
+++ b/Engine/Source/Programs/GPUReshape/Resources/Workspaces/BasicWorkspace.json
@@ -0,0 +1,12 @@
+{
+  "Config": {
+    "Coverage": true,
+    "Detail": true
+  },
+  "Features": [
+    "Resource Bounds",
+    "Export Stability",
+    "Descriptor",
+    "Loop"
+  ]
+}
\ No newline at end of file
diff --git a/Engine/Source/Programs/GPUReshape/Source/GPUReshapeBootstrapper/Private/GPUReshapeBootstrapper.cpp b/Engine/Source/Programs/GPUReshape/Source/GPUReshapeBootstrapper/Private/GPUReshapeBootstrapper.cpp
--- a/Engine/Source/Programs/GPUReshape/Source/GPUReshapeBootstrapper/Private/GPUReshapeBootstrapper.cpp
+++ b/Engine/Source/Programs/GPUReshape/Source/GPUReshapeBootstrapper/Private/GPUReshapeBootstrapper.cpp
@@ -45,18 +45,17 @@ INT32_MAIN_INT32_ARGC_TCHAR_ARGV()
 	}
 #endif // !UE_BUILD_SHIPPING
 
-	
-	FString GPUReshapePath;
-	if (!FParse::Value(FCommandLine::Get(), TEXT("GRS.Path="), GPUReshapePath))
+	FString BootstrapTarget;
+	if (!FParse::Value(FCommandLine::Get(), TEXT("BootstrapTarget="), BootstrapTarget))
 	{
-		UE_LOG(LogGPUReshape, Error, TEXT("GPU Reshape path not set"));
+		UE_LOG(LogGPUReshape, Error, TEXT("Target executable path not set"));
 		return 1u;
 	}
 	
-	FString AppPath;
-	if (!FParse::Value(FCommandLine::Get(), TEXT("GRS.App="), AppPath))
+	FString GPUReshapePath;
+	if (!FParse::Value(FCommandLine::Get(), TEXT("GRS.Path="), GPUReshapePath))
 	{
-		UE_LOG(LogGPUReshape, Error, TEXT("Target executable path not set"));
+		UE_LOG(LogGPUReshape, Error, TEXT("GPU Reshape path not set"));
 		return 1u;
 	}
 	
@@ -77,13 +76,17 @@ INT32_MAIN_INT32_ARGC_TCHAR_ARGV()
 	int32 Timeout = 7200;
 	FParse::Value(FCommandLine::Get(), TEXT("GRS.Timeout="), Timeout);
 	
+	FString SymbolPath = "";
+	FParse::Value(FCommandLine::Get(), TEXT("GRS.SymbolPath="), SymbolPath);
+	
 	// Start in headless mode
 	FStringBuilderBase GRSCommandLine;
 	GRSCommandLine << "headless ";
 	GRSCommandLine << "-report " << "\"" << ReportPath << "\" ";
 	GRSCommandLine << "-workspace \"" << WorkspacePath << "\" ";
 	GRSCommandLine << "-timeout " << Timeout << " ";
-	GRSCommandLine << "-app " << AppPath;
+	GRSCommandLine << "-symbol " << SymbolPath << " ";
+	GRSCommandLine << "-app " << BootstrapTarget;
 
 	for (int32_t i = 1; i < ArgC; i++)
 	{
