Expose GPU-Reshape workspace from perf options

diff --git a/Engine/Plugins/Performance/AutomatedPerfTesting/Build/Inc/AutomatedPerfTestCommonSettings.xml b/Engine/Plugins/Performance/AutomatedPerfTesting/Build/Inc/AutomatedPerfTestCommonSettings.xml
--- a/Engine/Plugins/Performance/AutomatedPerfTesting/Build/Inc/AutomatedPerfTestCommonSettings.xml
+++ b/Engine/Plugins/Performance/AutomatedPerfTesting/Build/Inc/AutomatedPerfTestCommonSettings.xml
@@ -47,6 +47,7 @@
 	<Option Name="DoGPUPerf" DefaultValue="false" Restrict="true|false" Description="Whether to run GPUPerf (GPU performance) test." />
 	<Option Name="DoInsightsTrace" DefaultValue="false" Restrict="true|false" Description="Whether or not to capture Unreal Insights trace during the run." />
 	<Option Name="DoGPUReshape" DefaultValue="false" Restrict="true|false" Description="Whether to run GPU instrumentation test with GPU Reshape." />
+	<Option Name="GPUReshapeWorkspace" DefaultValue="BasicWorkspace.json" Description="User workspace configuration for GPU Reshape." />
 
 	<!-- Construct the common arguments for the performance test and report commands -->
 	<Property Name="AutomatedPerformanceCommonArgs" Value = "-AutomatedPerfTest.UseShippingInsights=$(UseShippingInsights)" />
@@ -74,7 +75,7 @@
 		<Property Name="GPUReshapeBootstrapPath"  Value="$(RootDir)\Engine\Source\Programs\GPUReshape"/>
 		<Property Name="GPUReshapeReportPath"     Value="GRS.Report.json"/>
 		<Property Name="GPUReshapeReportPath"     Value="$(PerfOutputPath)\$(GPUReshapeReportPath)" If="'$(PerfOutputPath)' != ''"/>
-		<Property Name="GPUReshapeWorkspacePath"  Value="$(GPUReshapeBootstrapPath)\Resources\Workspaces\BasicWorkspace.json"/>
+		<Property Name="GPUReshapeWorkspacePath"  Value="$(GPUReshapeBootstrapPath)\Resources\Workspaces\$(GPUReshapeWorkspace)"/>
 		<Property Name="GPUReshapeSymbolPath"     Value="$(RootDir)\$(ProjectPath)\Saved\ShaderSymbols"/>
 		
 		<!-- Bootstrap arguments -->
