APT: Fix insights trace file name

* Fix insights trace file not being found by the APT Gauntlet node
* Add separate virtual functions to specify the CSV and Utrace filenames which each test controller can override if needed.

diff --git a/Engine/Plugins/Performance/AutomatedPerfTesting/Build/Scripts/AutomatedPerfTestNode.cs b/Engine/Plugins/Performance/AutomatedPerfTesting/Build/Scripts/AutomatedPerfTestNode.cs
--- a/Engine/Plugins/Performance/AutomatedPerfTesting/Build/Scripts/AutomatedPerfTestNode.cs
+++ b/Engine/Plugins/Performance/AutomatedPerfTesting/Build/Scripts/AutomatedPerfTestNode.cs
@@ -140,6 +140,24 @@ namespace AutomatedPerfTest
 			HandledErrors = new List<HandledError>();
 		}
 
+		protected virtual string GetNormalizedInsightsFileName(string FileName) => FileName.Replace(".csv", ".utrace");
+
+		protected string GetNormalizedInsightsFileName(string CSVFileName, string TestTypeLiteral)
+		{
+			int Index = CSVFileName.IndexOf(TestTypeLiteral);
+
+			if (Index < 0)
+			{
+				Index = CSVFileName.ToLower().IndexOf(".csv");
+			}
+			else
+			{
+				Index = Index + TestTypeLiteral.Length;
+			}
+
+			return CSVFileName.Substring(0, Index) + ".utrace";
+		}
+
 		/// <summary>
 		/// Periodically called while test is running. Updates logs.
 		/// </summary>
@@ -375,8 +393,9 @@ namespace AutomatedPerfTest
 							IEnumerable<CsvImportEntry> CsvImportEntries = ImportEntries as CsvImportEntry[] ?? ImportEntries.ToArray();
 							if (GetConfiguration().DoInsightsTrace)
 							{
-								string InsightsFilename = Path.GetFileNameWithoutExtension(CsvImportEntries.First().CsvFilename)
-									.Replace(".csv", ".utrace");
+								string InsightsFilename = Path.GetFileNameWithoutExtension(CsvImportEntries.First().CsvFilename);
+
+								InsightsFilename = GetNormalizedInsightsFileName(InsightsFilename);
 
 								// recursively look for trace files that match the CSV's filename in the artifact path
 								string[] MatchingTraces = FindFiles($"*{InsightsFilename}", true, ArtifactPath);
@@ -1072,6 +1091,11 @@ namespace AutomatedPerfTest
 
 			return OutSequenceList;
 		}
+
+		protected override string GetNormalizedInsightsFileName(string CSVFileName)
+		{
+			return GetNormalizedInsightsFileName(CSVFileName, "_Sequence");
+		}
 	}
 
 	/// <summary>
@@ -1230,6 +1254,10 @@ namespace AutomatedPerfTest
 
 			return OutReplayList;
 		}
+		protected override string GetNormalizedInsightsFileName(string CSVFileName)
+		{
+			return GetNormalizedInsightsFileName(CSVFileName, "_Replay");
+		}
 
 		private readonly string ReplayDataSourceName = "ReplayRun";
 		private readonly string ReplayParamName = "AutomatedPerfTest.ReplayPerfTest.ReplayName";
@@ -1284,6 +1312,11 @@ namespace AutomatedPerfTest
 
 			return Config;
 		}
+
+		protected override string GetNormalizedInsightsFileName(string CSVFileName)
+		{
+			return GetNormalizedInsightsFileName(CSVFileName, "_StaticCamera");
+		}
 	}
 
 	/// <summary>
@@ -1348,6 +1381,11 @@ namespace AutomatedPerfTest
 
 			return Config;
 		}
+
+		protected override string GetNormalizedInsightsFileName(string CSVFileName)
+		{
+			return GetNormalizedInsightsFileName(CSVFileName, "_Materials");
+		}
 	}
 
 	/// <summary>
diff --git a/Engine/Plugins/Performance/AutomatedPerfTesting/Source/AutomatedPerfTesting/Private/AutomatedPerfTestControllerBase.cpp b/Engine/Plugins/Performance/AutomatedPerfTesting/Source/AutomatedPerfTesting/Private/AutomatedPerfTestControllerBase.cpp
--- a/Engine/Plugins/Performance/AutomatedPerfTesting/Source/AutomatedPerfTesting/Private/AutomatedPerfTestControllerBase.cpp
+++ b/Engine/Plugins/Performance/AutomatedPerfTesting/Source/AutomatedPerfTesting/Private/AutomatedPerfTestControllerBase.cpp
@@ -184,7 +184,7 @@ bool UAutomatedPerfTestControllerBase::RequestsLockedDynRes() const
 
 bool UAutomatedPerfTestControllerBase::TryStartInsightsTrace()
 {
-	const FString TraceFileName = GetTestID() + ".utrace";
+	const FString TraceFileName = GetInsightsFilename() + ".utrace";
 	UE_LOG(LogAutomatedPerfTest, Log, TEXT("Attempting to start insights trace to file %s with channels %s"), *TraceFileName, *GetTraceChannels());
 	return FTraceAuxiliary::Start(FTraceAuxiliary::EConnectionType::File, *TraceFileName, *GetTraceChannels());
 }
@@ -200,7 +200,7 @@ bool UAutomatedPerfTestControllerBase::TryStopInsightsTrace()
 
 bool UAutomatedPerfTestControllerBase::TryStartCSVProfiler()
 {
-	return TryStartCSVProfiler(GetTestID());
+	return TryStartCSVProfiler(GetCSVFilename());
 }
 
 bool UAutomatedPerfTestControllerBase::TryStartCSVProfiler(FString CSVFileName)
@@ -534,6 +534,16 @@ void UAutomatedPerfTestControllerBase::MarkProfilingEnd()
 	}
 }
 
+FString UAutomatedPerfTestControllerBase::GetInsightsFilename()
+{
+	return GetTestID();
+}
+
+FString UAutomatedPerfTestControllerBase::GetCSVFilename()
+{
+	return GetTestID();
+}
+
 void UAutomatedPerfTestControllerBase::SetCSVOutputMode(EAutomatedPerfTestCSVOutputMode NewOutputMode)
 {
 	CSVOutputMode = NewOutputMode;	
diff --git a/Engine/Plugins/Performance/AutomatedPerfTesting/Source/AutomatedPerfTesting/Private/AutomatedSequencePerfTest.cpp b/Engine/Plugins/Performance/AutomatedPerfTesting/Source/AutomatedPerfTesting/Private/AutomatedSequencePerfTest.cpp
--- a/Engine/Plugins/Performance/AutomatedPerfTesting/Source/AutomatedPerfTesting/Private/AutomatedSequencePerfTest.cpp
+++ b/Engine/Plugins/Performance/AutomatedPerfTesting/Source/AutomatedPerfTesting/Private/AutomatedSequencePerfTest.cpp
@@ -277,7 +277,7 @@ void UAutomatedSequencePerfTest::OnCameraCut(UCameraComponent* CameraComponent)
 
 FString UAutomatedSequencePerfTest::GetTestID()
 {
-	return CurrentMapSequenceCombo.IsSet() ? Super::GetTestID() + "_Sequence_" + CurrentMapSequenceCombo->ComboName.ToString() : Super::GetTestID();
+	return Super::GetTestID() + "_Sequence";
 }
 
 FString UAutomatedSequencePerfTest::GetCameraCutID()
@@ -369,3 +369,14 @@ void UAutomatedSequencePerfTest::UnbindAllDelegates()
 		World->GetTimerManager().ClearAllTimersForObject(this);
 	}
 }
+
+FString UAutomatedSequencePerfTest::GetCSVFilename()
+{
+	if (!CurrentMapSequenceCombo.IsSet() || CurrentMapSequenceCombo->ComboName.ToString().IsEmpty())
+	{
+		UE_LOG(LogAutomatedPerfTest, Error, TEXT("Current Map Sequence Combo Name not set"));
+		return Super::GetCSVFilename();
+	}
+
+	return GetTestID() + "_" + CurrentMapSequenceCombo->ComboName.ToString();
+}
diff --git a/Engine/Plugins/Performance/AutomatedPerfTesting/Source/AutomatedPerfTesting/Private/StaticCameraTests/AutomatedStaticCameraPerfTestBase.cpp b/Engine/Plugins/Performance/AutomatedPerfTesting/Source/AutomatedPerfTesting/Private/StaticCameraTests/AutomatedStaticCameraPerfTestBase.cpp
--- a/Engine/Plugins/Performance/AutomatedPerfTesting/Source/AutomatedPerfTesting/Private/StaticCameraTests/AutomatedStaticCameraPerfTestBase.cpp
+++ b/Engine/Plugins/Performance/AutomatedPerfTesting/Source/AutomatedPerfTesting/Private/StaticCameraTests/AutomatedStaticCameraPerfTestBase.cpp
@@ -97,7 +97,17 @@ void UAutomatedStaticCameraPerfTestBase::RunTest()
 
 FString UAutomatedStaticCameraPerfTestBase::GetTestID()
 {
-	return Super::GetTestID() + "_StaticCamera_" + GetCurrentMap();
+	return Super::GetTestID() + "_StaticCamera";
+}
+
+FString UAutomatedStaticCameraPerfTestBase::GetCSVFilename()
+{
+	if (GetCurrentMap() != CurrentMapPath.GetAssetName())
+	{
+		UE_LOG(LogAutomatedPerfTest, Warning, TEXT("Current Map is not the expected path. Current: %s, Expected: %s"), *GetCurrentMap(), *CurrentMapPath.GetAssetName());
+	}
+	
+	return GetTestID() + "_" + GetCurrentMap();
 }
 
 bool UAutomatedStaticCameraPerfTestBase::TryStartCSVProfiler(FString CSVFileName)
diff --git a/Engine/Plugins/Performance/AutomatedPerfTesting/Source/AutomatedPerfTesting/Public/AutomatedPerfTestControllerBase.h b/Engine/Plugins/Performance/AutomatedPerfTesting/Source/AutomatedPerfTesting/Public/AutomatedPerfTestControllerBase.h
--- a/Engine/Plugins/Performance/AutomatedPerfTesting/Source/AutomatedPerfTesting/Public/AutomatedPerfTestControllerBase.h
+++ b/Engine/Plugins/Performance/AutomatedPerfTesting/Source/AutomatedPerfTesting/Public/AutomatedPerfTestControllerBase.h
@@ -138,6 +138,9 @@ protected:
 	UE_API void TeardownProfiling();
 	UE_API void MarkProfilingStart();
 	UE_API void MarkProfilingEnd();
+	
+	UE_API virtual FString GetInsightsFilename();
+	UE_API virtual FString GetCSVFilename();
 	/* End Profiling Functions*/
 
 private:
diff --git a/Engine/Plugins/Performance/AutomatedPerfTesting/Source/AutomatedPerfTesting/Public/AutomatedSequencePerfTest.h b/Engine/Plugins/Performance/AutomatedPerfTesting/Source/AutomatedPerfTesting/Public/AutomatedSequencePerfTest.h
--- a/Engine/Plugins/Performance/AutomatedPerfTesting/Source/AutomatedPerfTesting/Public/AutomatedSequencePerfTest.h
+++ b/Engine/Plugins/Performance/AutomatedPerfTesting/Source/AutomatedPerfTesting/Public/AutomatedSequencePerfTest.h
@@ -118,6 +118,7 @@ public:
 protected:
 	UE_API virtual void OnInit() override;
 	UE_API virtual void UnbindAllDelegates() override;
+	UE_API virtual FString GetCSVFilename() override;
 
 private:
 	const UAutomatedSequencePerfTestProjectSettings* Settings;
diff --git a/Engine/Plugins/Performance/AutomatedPerfTesting/Source/AutomatedPerfTesting/Public/StaticCameraTests/AutomatedStaticCameraPerfTestBase.h b/Engine/Plugins/Performance/AutomatedPerfTesting/Source/AutomatedPerfTesting/Public/StaticCameraTests/AutomatedStaticCameraPerfTestBase.h
--- a/Engine/Plugins/Performance/AutomatedPerfTesting/Source/AutomatedPerfTesting/Public/StaticCameraTests/AutomatedStaticCameraPerfTestBase.h
+++ b/Engine/Plugins/Performance/AutomatedPerfTesting/Source/AutomatedPerfTesting/Public/StaticCameraTests/AutomatedStaticCameraPerfTestBase.h
@@ -120,6 +120,8 @@ protected:
 	UE_API virtual void UnbindAllDelegates() override;
 	UE_API virtual void TriggerExitAfterDelay() override;
 
+	UE_API virtual FString GetCSVFilename() override;
+
 private:
 	UPROPERTY()
 	TArray<TObjectPtr<ACameraActor>> CamerasToTest;
