Updated GPU Reshape branching paths
- Added GRSBranch command line option

diff --git a/Engine/Plugins/Developer/GPUReshape/Source/GPUReshape/Private/GPUReshapeModule.cpp b/Engine/Plugins/Developer/GPUReshape/Source/GPUReshape/Private/GPUReshapeModule.cpp
--- a/Engine/Plugins/Developer/GPUReshape/Source/GPUReshape/Private/GPUReshapeModule.cpp
+++ b/Engine/Plugins/Developer/GPUReshape/Source/GPUReshape/Private/GPUReshapeModule.cpp
@@ -131,8 +131,12 @@ bool FGPUReshapeModule::FindAndInstallLoader()
 	// Allow binary overrides (useful for development)
 	if (!FParse::Value(FCommandLine::Get(), TEXT("GRSPath"), GPUReshapePath))
 	{
+		// Default to the raytracing branch, until it's been merged back to dev
+		FString Branch = "Raytracing";
+		FParse::Value(FCommandLine::Get(), TEXT("GRSBranch"), GPUReshapePath);
+	
 		// Default path, for now at least
-		GPUReshapePath = FPaths::EngineDir() / TEXT("Restricted/NotForLicensees/Binaries/ThirdParty/GPUReshape/Win64");
+		GPUReshapePath = FPaths::EngineDir() / TEXT("Restricted/NotForLicensees/Binaries/ThirdParty/GPUReshape/Branch") / Branch / TEXT("Win64");
 	}
 
 	// Expected loader, handles all backends
diff --git a/Engine/Plugins/Performance/AutomatedPerfTesting/Build/Inc/AutomatedPerfTestCommonSettings.xml b/Engine/Plugins/Performance/AutomatedPerfTesting/Build/Inc/AutomatedPerfTestCommonSettings.xml
--- a/Engine/Plugins/Performance/AutomatedPerfTesting/Build/Inc/AutomatedPerfTestCommonSettings.xml
+++ b/Engine/Plugins/Performance/AutomatedPerfTesting/Build/Inc/AutomatedPerfTestCommonSettings.xml
@@ -71,7 +71,7 @@
 	
 	<!-- Reshape configuration -->
 	<Do If="$(DoGPUReshape)">
-		<Property Name="GPUReshapePath"           Value="$(RootDir)\Engine\Restricted\NotForLicensees\Binaries\ThirdParty\GPUReshape\Win64"/>
+		<Property Name="GPUReshapePath"           Value="$(RootDir)\Engine\Restricted\NotForLicensees\Binaries\ThirdParty\GPUReshape\Branch\Raytracing\Win64"/>
 		<Property Name="GPUReshapeBootstrapPath"  Value="$(RootDir)\Engine\Source\Programs\GPUReshape"/>
 		<Property Name="GPUReshapeReportPath"     Value="$(GPUReshapeWorkspace).GRS.Report.json"/>
 		<Property Name="GPUReshapeReportPath"     Value="$(PerfOutputPath)\$(GPUReshapeReportPath)" If="'$(PerfOutputPath)' != ''"/>
diff --git a/Engine/Plugins/Performance/AutomatedPerfTesting/Build/Scripts/AutomatedPerfTestNode.cs b/Engine/Plugins/Performance/AutomatedPerfTesting/Build/Scripts/AutomatedPerfTestNode.cs
--- a/Engine/Plugins/Performance/AutomatedPerfTesting/Build/Scripts/AutomatedPerfTestNode.cs
+++ b/Engine/Plugins/Performance/AutomatedPerfTesting/Build/Scripts/AutomatedPerfTestNode.cs
@@ -858,7 +858,7 @@ namespace AutomatedPerfTest
 			InternalUtils.SafeCreateDirectory(OutputPath, true);
 			
 			// Hardcoded path, for now
-			var GPUReshapePath = Path.Combine(Unreal.RootDirectory.FullName, "Engine/Restricted/NotForLicensees/Binaries/ThirdParty/GPUReshape/Win64/GPUReshape.exe");
+			var GPUReshapePath = Path.Combine(Unreal.RootDirectory.FullName, "Engine/Restricted/NotForLicensees/Binaries/ThirdParty/GPUReshape/Win64/Branch/Raytracing/GPUReshape.exe");
 			
 			// Get all reports, search each test case
 			var Reports = Directory.GetFiles(ArtifactPath, "*GRS.Report.json", SearchOption.AllDirectories).ToList();
