Automated Testing CookOnTheFly BuildGraph fix

diff --git a/Engine/Plugins/Performance/AutomatedPerfTesting/Build/Inc/AutomatedPerfTestCommonSettings.xml b/Engine/Plugins/Performance/AutomatedPerfTesting/Build/Inc/AutomatedPerfTestCommonSettings.xml
--- a/Engine/Plugins/Performance/AutomatedPerfTesting/Build/Inc/AutomatedPerfTestCommonSettings.xml
+++ b/Engine/Plugins/Performance/AutomatedPerfTesting/Build/Inc/AutomatedPerfTestCommonSettings.xml
@@ -97,38 +97,38 @@
 	<!-- append plugin-provided tests to GauntletSettings' TargetTestList -->
 	<Do If="$(WithAPTSequenceTests)">
 		<ForEach Name="MapSequenceComboName" Values="$(MapSequenceComboNames)" Separator="+">
-			<Property Name="TargetTestList" Value="$(TargetTestList)AutomatedPerfTest.SequenceTest{-AutomatedPerfTest.SequencePerfTest.MapSequenceName=$(MapSequenceComboName) -AutomatedPerfTest.DoCSVProfiler $(PerfTestCommonArgs)}[Sequence Perf $(MapSequenceComboName)]+" If="$(DoPerf)" />
-			<Property Name="TargetTestList" Value="$(TargetTestList)AutomatedPerfTest.SequenceTest{-AutomatedPerfTest.SequencePerfTest.MapSequenceName=$(MapSequenceComboName) -AutomatedPerfTest.DoCSVProfiler -AutomatedPerfTest.DoLLM $(LLMTestCommonArgs)}[Sequence LLM $(MapSequenceComboName)]+" If="$(DoLLM)"/>
-			<Property Name="TargetTestList" Value="$(TargetTestList)AutomatedPerfTest.SequenceTest{-AutomatedPerfTest.SequencePerfTest.MapSequenceName=$(MapSequenceComboName) -AutomatedPerfTest.DoCSVProfiler -AutomatedPerfTest.DoGPUPerf $(GPUPerfTestCommonArgs)}[Sequence GPUPerf $(MapSequenceComboName)]+" If="$(DoGPUPerf)" />
-			<Property Name="TargetTestList" Value="$(TargetTestList)AutomatedPerfTest.SequenceTest{-AutomatedPerfTest.SequencePerfTest.MapSequenceName=$(MapSequenceComboName) -AutomatedPerfTest.DoCSVProfiler -AutomatedPerfTest.DoInsightsTrace $(InsightsTestCommonArgs)}[Sequence Insights $(MapSequenceComboName)]+" If="$(DoInsightsTrace)"/>
-			<Property Name="TargetTestList" Value="$(TargetTestList)AutomatedPerfTest.SequenceTest{-AutomatedPerfTest.SequencePerfTest.MapSequenceName=$(MapSequenceComboName) -AutomatedPerfTest.DoGPUReshape $(GPUReshapeTestCommonArgs)}[Sequence GPUReshape $(MapSequenceComboName)]+" If="$(DoGPUReshape)"/>
+			<Property Name="TargetTestList" Value="$(TargetTestList)+AutomatedPerfTest.SequenceTest{-AutomatedPerfTest.SequencePerfTest.MapSequenceName=$(MapSequenceComboName) -AutomatedPerfTest.DoCSVProfiler $(PerfTestCommonArgs)}[Sequence Perf $(MapSequenceComboName)]" If="$(DoPerf)" />
+			<Property Name="TargetTestList" Value="$(TargetTestList)+AutomatedPerfTest.SequenceTest{-AutomatedPerfTest.SequencePerfTest.MapSequenceName=$(MapSequenceComboName) -AutomatedPerfTest.DoCSVProfiler -AutomatedPerfTest.DoLLM $(LLMTestCommonArgs)}[Sequence LLM $(MapSequenceComboName)]" If="$(DoLLM)"/>
+			<Property Name="TargetTestList" Value="$(TargetTestList)+AutomatedPerfTest.SequenceTest{-AutomatedPerfTest.SequencePerfTest.MapSequenceName=$(MapSequenceComboName) -AutomatedPerfTest.DoCSVProfiler -AutomatedPerfTest.DoGPUPerf $(GPUPerfTestCommonArgs)}[Sequence GPUPerf $(MapSequenceComboName)]" If="$(DoGPUPerf)" />
+			<Property Name="TargetTestList" Value="$(TargetTestList)+AutomatedPerfTest.SequenceTest{-AutomatedPerfTest.SequencePerfTest.MapSequenceName=$(MapSequenceComboName) -AutomatedPerfTest.DoCSVProfiler -AutomatedPerfTest.DoInsightsTrace $(InsightsTestCommonArgs)}[Sequence Insights $(MapSequenceComboName)]" If="$(DoInsightsTrace)"/>
+			<Property Name="TargetTestList" Value="$(TargetTestList)+AutomatedPerfTest.SequenceTest{-AutomatedPerfTest.SequencePerfTest.MapSequenceName=$(MapSequenceComboName) -AutomatedPerfTest.DoGPUReshape $(GPUReshapeTestCommonArgs)}[Sequence GPUReshape $(MapSequenceComboName)]" If="$(DoGPUReshape)"/>
 		</ForEach>
 	</Do>
 
 	<Do If="$(WithAPTStaticCameraTests)">
 		<ForEach Name="StaticCameraMapName" Values="$(StaticCameraMapNames)" Separator="+">
-			<Property Name="TargetTestList" Value="$(TargetTestList)AutomatedPerfTest.StaticCameraTest{-AutomatedPerfTest.StaticCameraPerfTest.MapName=$(StaticCameraMapName) -AutomatedPerfTest.DoCSVProfiler $(PerfTestCommonArgs)}[Static Camera Perf $(StaticCameraMapName)]+" If="$(DoPerf)" />
-			<Property Name="TargetTestList" Value="$(TargetTestList)AutomatedPerfTest.StaticCameraTest{-AutomatedPerfTest.StaticCameraPerfTest.MapName=$(StaticCameraMapName) -AutomatedPerfTest.DoCSVProfiler -AutomatedPerfTest.DoLLM $(LLMTestCommonArgs)}[Static Camera LLM $(StaticCameraMapName)]+" If="$(DoLLM)"/>
-			<Property Name="TargetTestList" Value="$(TargetTestList)AutomatedPerfTest.StaticCameraTest{-AutomatedPerfTest.StaticCameraPerfTest.MapName=$(StaticCameraMapName) -AutomatedPerfTest.DoCSVProfiler -AutomatedPerfTest.DoGPUPerf $(LLMTestCommonArgs)}[Static Camera GPUPerf $(StaticCameraMapName)]+" If="$(DoGPUPerf)"/>
-			<Property Name="TargetTestList" Value="$(TargetTestList)AutomatedPerfTest.StaticCameraTest{-AutomatedPerfTest.StaticCameraPerfTest.MapName=$(StaticCameraMapName) -AutomatedPerfTest.DoCSVProfiler -AutomatedPerfTest.DoInsightsTrace $(InsightsTestCommonArgs)}[Static Camera Insights $(StaticCameraMapName)]+" If="$(DoInsightsTrace)"/>
-			<Property Name="TargetTestList" Value="$(TargetTestList)AutomatedPerfTest.StaticCameraTest{-AutomatedPerfTest.StaticCameraPerfTest.MapName=$(StaticCameraMapName) -AutomatedPerfTest.DoGPUReshape $(GPUReshapeTestCommonArgs)}[Static Camera GPUReshape $(StaticCameraMapName)]+" If="$(DoGPUReshape)"/>
+			<Property Name="TargetTestList" Value="$(TargetTestList)+AutomatedPerfTest.StaticCameraTest{-AutomatedPerfTest.StaticCameraPerfTest.MapName=$(StaticCameraMapName) -AutomatedPerfTest.DoCSVProfiler $(PerfTestCommonArgs)}[Static Camera Perf $(StaticCameraMapName)]" If="$(DoPerf)" />
+			<Property Name="TargetTestList" Value="$(TargetTestList)+AutomatedPerfTest.StaticCameraTest{-AutomatedPerfTest.StaticCameraPerfTest.MapName=$(StaticCameraMapName) -AutomatedPerfTest.DoCSVProfiler -AutomatedPerfTest.DoLLM $(LLMTestCommonArgs)}[Static Camera LLM $(StaticCameraMapName)]" If="$(DoLLM)"/>
+			<Property Name="TargetTestList" Value="$(TargetTestList)+AutomatedPerfTest.StaticCameraTest{-AutomatedPerfTest.StaticCameraPerfTest.MapName=$(StaticCameraMapName) -AutomatedPerfTest.DoCSVProfiler -AutomatedPerfTest.DoGPUPerf $(LLMTestCommonArgs)}[Static Camera GPUPerf $(StaticCameraMapName)]" If="$(DoGPUPerf)"/>
+			<Property Name="TargetTestList" Value="$(TargetTestList)+AutomatedPerfTest.StaticCameraTest{-AutomatedPerfTest.StaticCameraPerfTest.MapName=$(StaticCameraMapName) -AutomatedPerfTest.DoCSVProfiler -AutomatedPerfTest.DoInsightsTrace $(InsightsTestCommonArgs)}[Static Camera Insights $(StaticCameraMapName)]" If="$(DoInsightsTrace)"/>
+			<Property Name="TargetTestList" Value="$(TargetTestList)+AutomatedPerfTest.StaticCameraTest{-AutomatedPerfTest.StaticCameraPerfTest.MapName=$(StaticCameraMapName) -AutomatedPerfTest.DoGPUReshape $(GPUReshapeTestCommonArgs)}[Static Camera GPUReshape $(StaticCameraMapName)]" If="$(DoGPUReshape)"/>
 		</ForEach>
 	</Do>
 
 	<Do If="$(WithAPTMaterialTest)">
-		<Property Name="TargetTestList" Value="$(TargetTestList)AutomatedPerfTest.MaterialTest{-AutomatedPerfTest.DoCSVProfiler $(PerfTestCommonArgs)}[Material Perf]+" If="$(DoPerf)" />
-		<Property Name="TargetTestList" Value="$(TargetTestList)AutomatedPerfTest.MaterialTest{-AutomatedPerfTest.DoCSVProfiler -AutomatedPerfTest.DoLLM $(LLMTestCommonArgs)}[Material LLM]+" If="$(DoLLM)"/>
-		<Property Name="TargetTestList" Value="$(TargetTestList)AutomatedPerfTest.MaterialTest{-AutomatedPerfTest.DoCSVProfiler -AutomatedPerfTest.DoGPUPerf $(LLMTestCommonArgs)}[Material GPUPerf]+" If="$(DoGPUPerf)"/>
-		<Property Name="TargetTestList" Value="$(TargetTestList)AutomatedPerfTest.MaterialTest{-AutomatedPerfTest.DoCSVProfiler -AutomatedPerfTest.DoInsightsTrace $(InsightsTestCommonArgs)}[Material Insights]+" If="$(DoInsightsTrace)"/>
-		<Property Name="TargetTestList" Value="$(TargetTestList)AutomatedPerfTest.MaterialTest{-AutomatedPerfTest.DoCSVProfiler -AutomatedPerfTest.DoGPUReshape $(GPUReshapeTestCommonArgs)}[Material GPUReshape]+" If="$(DoGPUReshape)"/>
+		<Property Name="TargetTestList" Value="$(TargetTestList)+AutomatedPerfTest.MaterialTest{-AutomatedPerfTest.DoCSVProfiler $(PerfTestCommonArgs)}[Material Perf]" If="$(DoPerf)" />
+		<Property Name="TargetTestList" Value="$(TargetTestList)+AutomatedPerfTest.MaterialTest{-AutomatedPerfTest.DoCSVProfiler -AutomatedPerfTest.DoLLM $(LLMTestCommonArgs)}[Material LLM]" If="$(DoLLM)"/>
+		<Property Name="TargetTestList" Value="$(TargetTestList)+AutomatedPerfTest.MaterialTest{-AutomatedPerfTest.DoCSVProfiler -AutomatedPerfTest.DoGPUPerf $(LLMTestCommonArgs)}[Material GPUPerf]" If="$(DoGPUPerf)"/>
+		<Property Name="TargetTestList" Value="$(TargetTestList)+AutomatedPerfTest.MaterialTest{-AutomatedPerfTest.DoCSVProfiler -AutomatedPerfTest.DoInsightsTrace $(InsightsTestCommonArgs)}[Material Insights]" If="$(DoInsightsTrace)"/>
+		<Property Name="TargetTestList" Value="$(TargetTestList)+AutomatedPerfTest.MaterialTest{-AutomatedPerfTest.DoCSVProfiler -AutomatedPerfTest.DoGPUReshape $(GPUReshapeTestCommonArgs)}[Material GPUReshape]" If="$(DoGPUReshape)"/>
 	</Do>
 
 	<Do If="$(WithAPTReplayTest)">
-		<Property Name="TargetTestList" Value="$(TargetTestList)AutomatedPerfTest.ReplayTest{-AutomatedPerfTest.ReplayPerfTest.ReplayName=&quot;$(ReplayName)&quot; -AutomatedPerfTest.DoCSVProfiler $(PerfTestCommonArgs)}[Replay Perf]+" If="$(DoPerf)" />
-		<Property Name="TargetTestList" Value="$(TargetTestList)AutomatedPerfTest.ReplayTest{-AutomatedPerfTest.ReplayPerfTest.ReplayName=&quot;$(ReplayName)&quot; -AutomatedPerfTest.DoCSVProfiler -AutomatedPerfTest.DoLLM $(LLMTestCommonArgs)}[Replay LLM]+" If="$(DoLLM)"/>
-		<Property Name="TargetTestList" Value="$(TargetTestList)AutomatedPerfTest.ReplayTest{-AutomatedPerfTest.ReplayPerfTest.ReplayName=&quot;$(ReplayName)&quot; -AutomatedPerfTest.DoCSVProfiler -AutomatedPerfTest.DoGPUPerf $(LLMTestCommonArgs)}[Replay GPUPerf]+" If="$(DoGPUPerf)"/>
-		<Property Name="TargetTestList" Value="$(TargetTestList)AutomatedPerfTest.ReplayTest{-AutomatedPerfTest.ReplayPerfTest.ReplayName=&quot;$(ReplayName)&quot; -AutomatedPerfTest.DoCSVProfiler -AutomatedPerfTest.DoInsightsTrace $(InsightsTestCommonArgs)}[Replay Insights]+" If="$(DoInsightsTrace)"/>
-		<Property Name="TargetTestList" Value="$(TargetTestList)AutomatedPerfTest.ReplayTest{-AutomatedPerfTest.ReplayPerfTest.ReplayName=&quot;$(ReplayName)&quot; -AutomatedPerfTest.DoGPUReshape $(GPUReshapeTestCommonArgs)}[Replay GPUReshape]+" If="$(DoGPUReshape)"/>
+		<Property Name="TargetTestList" Value="$(TargetTestList)+AutomatedPerfTest.ReplayTest{-AutomatedPerfTest.ReplayPerfTest.ReplayName=&quot;$(ReplayName)&quot; -AutomatedPerfTest.DoCSVProfiler $(PerfTestCommonArgs)}[Replay Perf]" If="$(DoPerf)" />
+		<Property Name="TargetTestList" Value="$(TargetTestList)+AutomatedPerfTest.ReplayTest{-AutomatedPerfTest.ReplayPerfTest.ReplayName=&quot;$(ReplayName)&quot; -AutomatedPerfTest.DoCSVProfiler -AutomatedPerfTest.DoLLM $(LLMTestCommonArgs)}[Replay LLM]" If="$(DoLLM)"/>
+		<Property Name="TargetTestList" Value="$(TargetTestList)+AutomatedPerfTest.ReplayTest{-AutomatedPerfTest.ReplayPerfTest.ReplayName=&quot;$(ReplayName)&quot; -AutomatedPerfTest.DoCSVProfiler -AutomatedPerfTest.DoGPUPerf $(LLMTestCommonArgs)}[Replay GPUPerf]" If="$(DoGPUPerf)"/>
+		<Property Name="TargetTestList" Value="$(TargetTestList)+AutomatedPerfTest.ReplayTest{-AutomatedPerfTest.ReplayPerfTest.ReplayName=&quot;$(ReplayName)&quot; -AutomatedPerfTest.DoCSVProfiler -AutomatedPerfTest.DoInsightsTrace $(InsightsTestCommonArgs)}[Replay Insights]" If="$(DoInsightsTrace)"/>
+		<Property Name="TargetTestList" Value="$(TargetTestList)+AutomatedPerfTest.ReplayTest{-AutomatedPerfTest.ReplayPerfTest.ReplayName=&quot;$(ReplayName)&quot; -AutomatedPerfTest.DoGPUReshape $(GPUReshapeTestCommonArgs)}[Replay GPUReshape]" If="$(DoGPUReshape)"/>
 	</Do>
 
 </BuildGraph>
