GPU Reshape, fixed gauntlet path

diff --git a/Engine/Plugins/Performance/AutomatedPerfTesting/Build/Scripts/AutomatedPerfTestNode.cs b/Engine/Plugins/Performance/AutomatedPerfTesting/Build/Scripts/AutomatedPerfTestNode.cs
--- a/Engine/Plugins/Performance/AutomatedPerfTesting/Build/Scripts/AutomatedPerfTestNode.cs
+++ b/Engine/Plugins/Performance/AutomatedPerfTesting/Build/Scripts/AutomatedPerfTestNode.cs
@@ -858,7 +858,9 @@ namespace AutomatedPerfTest
 			InternalUtils.SafeCreateDirectory(OutputPath, true);
 			
 			// Hardcoded path, for now
-			var GPUReshapePath = Path.Combine(Unreal.RootDirectory.FullName, "Engine/Restricted/NotForLicensees/Binaries/ThirdParty/GPUReshape/Win64/Branch/Raytracing/GPUReshape.exe");
+			// Raytracing serves as the current development branch
+			string GPUReshapeBranch = "Raytracing";
+			string GPUReshapePath   = Path.Combine(Unreal.RootDirectory.FullName, $"Engine/Restricted/NotForLicensees/Binaries/ThirdParty/GPUReshape/Branch/{GPUReshapeBranch}/Win64/GPUReshape.exe");
 			
 			// Get all reports, search each test case
 			var Reports = Directory.GetFiles(ArtifactPath, "*GRS.Report.json", SearchOption.AllDirectories).ToList();
