Fix Replay perf test on some platforms for UE Samples

* Fix support for APT Replay test on newer platforms
* Refactor to use static helper functions from within `PlatformTargetSupport`

diff --git a/Engine/Plugins/Performance/AutomatedPerfTesting/Build/Scripts/AutomatedPerfTestNode.cs b/Engine/Plugins/Performance/AutomatedPerfTesting/Build/Scripts/AutomatedPerfTestNode.cs
--- a/Engine/Plugins/Performance/AutomatedPerfTesting/Build/Scripts/AutomatedPerfTestNode.cs
+++ b/Engine/Plugins/Performance/AutomatedPerfTesting/Build/Scripts/AutomatedPerfTestNode.cs
@@ -1147,12 +1147,6 @@ namespace AutomatedPerfTest
 			
 			Config.DataSourceName = Config.GetDataSourceName(Context.BuildInfo.ProjectName, ReplayDataSourceName);
 
-			// Since we do not have a device instance at this stage. These platform support
-			// instances would help if we would still like to know if the target platform has
-			// certain features or functionality we need to be aware of in order to pass the
-			// right configuration to Gauntlet.
-			IEnumerable<IPlatformTargetSupport> PlatformSupportInstances = Gauntlet.Utils.InterfaceHelpers.FindImplementations<IPlatformTargetSupport>();
-
 			IEnumerable<UnrealTestRole> ClientRoles = Config.GetRequiredRoles(UnrealTargetRole.Client);
 			// extend the role(s) that we initialized in the base class
 			if (ClientRoles.Any())
@@ -1163,10 +1157,6 @@ namespace AutomatedPerfTest
 					ClientRole.FilesToCopy.Clear();
 
 					UnrealTargetPlatform Platform = Context.GetRoleContext(UnrealTargetRole.Client).Platform;
-					IPlatformTargetSupport MountablePlatform = PlatformSupportInstances
-						.Where(TargetSupport => TargetSupport.IsHostMountingSupported() && TargetSupport.Platform == Platform)
-						.FirstOrDefault();
-
 					List<string> ReplaysFromConfig = GetTestsFromConfig();
 
 					// Replay name not provided
@@ -1210,10 +1200,10 @@ namespace AutomatedPerfTest
 						throw new AutomationException($"Replay file '{Config.ReplayName}' not found");
 					}
 
-					if (MountablePlatform != null) 
+					if (PlatformTargetSupport.IsHostMountingSupported(Platform)) 
 					{
 						// If current device supports host mounting of files, update the path and use host mounting 
-						ReplayName = MountablePlatform.GetHostMountedPath(ReplayName);
+						ReplayName = PlatformTargetSupport.GetHostMountedPath(Platform, ReplayName);
 						Log.Info($"Host Mountable Platform Detected. Updated Replay File Path to: {ReplayName} ");
 					}
 					else
diff --git a/Engine/Source/Programs/AutomationTool/Gauntlet/Framework/Base/Gauntlet.TargetDevice.cs b/Engine/Source/Programs/AutomationTool/Gauntlet/Framework/Base/Gauntlet.TargetDevice.cs
--- a/Engine/Source/Programs/AutomationTool/Gauntlet/Framework/Base/Gauntlet.TargetDevice.cs
+++ b/Engine/Source/Programs/AutomationTool/Gauntlet/Framework/Base/Gauntlet.TargetDevice.cs
@@ -7,6 +7,7 @@ using AutomationTool;
 using UnrealBuildTool;
 using System.Threading;
 using System.Text.RegularExpressions;
+using System.Linq;
 
 namespace Gauntlet
 {
@@ -252,6 +253,29 @@ namespace Gauntlet
 		UnrealTargetPlatform? Platform { get; }
 	}
 
+	/// <summary>
+	/// Static helper class for querying platform specific target support
+	/// </summary>
+	public static class PlatformTargetSupport
+	{
+		public static IPlatformTargetSupport Get(UnrealTargetPlatform Platform)
+		{
+			if (PlatformSupportInstances == null || PlatformSupportInstances.Count == 0)
+			{
+				PlatformSupportInstances = Utils.InterfaceHelpers.FindImplementations<IPlatformTargetSupport>().ToList();
+			}
+
+			return PlatformSupportInstances
+				.Where(x => x.Platform == Platform)
+				.FirstOrDefault();
+		}
+
+		public static string GetHostMountedPath(UnrealTargetPlatform Platform, string Path) => Get(Platform).GetHostMountedPath(Path);
+		public static bool IsHostMountingSupported(UnrealTargetPlatform Platform) => Get(Platform).IsHostMountingSupported();
+
+		private static List<IPlatformTargetSupport> PlatformSupportInstances;
+	}
+
 	/// <summary>
 	/// Default abstract implementation of IPlatformTargetSupport
 	/// so that target platforms can override only parts they need
