APT: Fix insights trace file not being discovered

diff --git a/Engine/Plugins/Performance/AutomatedPerfTesting/Build/Scripts/AutomatedPerfTestNode.cs b/Engine/Plugins/Performance/AutomatedPerfTesting/Build/Scripts/AutomatedPerfTestNode.cs
--- a/Engine/Plugins/Performance/AutomatedPerfTesting/Build/Scripts/AutomatedPerfTestNode.cs
+++ b/Engine/Plugins/Performance/AutomatedPerfTesting/Build/Scripts/AutomatedPerfTestNode.cs
@@ -338,9 +338,10 @@ namespace AutomatedPerfTest
 			
 			if (Result == TestResult.Passed)
 			{
+				string InsightsTraceFile = string.Empty;
 				if (Config.DoInsightsTrace)
 				{
-					CopyInsightsTraceToOutput(ArtifactPath, OutputPath);
+					CopyInsightsTraceToOutput(ArtifactPath, OutputPath, out InsightsTraceFile);
 				}
 				
 				if (GetCurrentPass() <= GetNumPasses() && Config.DoCSVProfiler)
@@ -395,7 +396,7 @@ namespace AutomatedPerfTest
 							{
 								string InsightsFilename = Path.GetFileNameWithoutExtension(CsvImportEntries.First().CsvFilename);
 
-								InsightsFilename = GetNormalizedInsightsFileName(InsightsFilename);
+								InsightsFilename = string.IsNullOrEmpty(InsightsTraceFile) ? GetNormalizedInsightsFileName(InsightsFilename) : InsightsTraceFile + ".utrace";
 
 								// recursively look for trace files that match the CSV's filename in the artifact path
 								string[] MatchingTraces = FindFiles($"*{InsightsFilename}", true, ArtifactPath);
@@ -927,8 +928,9 @@ namespace AutomatedPerfTest
 			return Report;
 		}
 		
-		public void CopyInsightsTraceToOutput(string FromArtifactPath, string ToOutputPath)
+		public void CopyInsightsTraceToOutput(string FromArtifactPath, string ToOutputPath, out string OutTraceName)
 		{
+			OutTraceName = string.Empty;
 			Log.Info("Copying test insights trace from artifact path to report cache");
 			
 			// find all the available trace paths
@@ -969,6 +971,7 @@ namespace AutomatedPerfTest
 				{
 					InternalUtils.SafeCreateDirectory(Path.GetDirectoryName(PerfTracePath), true);
 					File.Copy(SortedTraces[0], PerfTracePath);
+					OutTraceName = Filename;
 				}
 				catch (Exception e)
 				{
