Add -keepscreenawake commandline parameter for automation to disable screen saver on platforms different than desktop
Remove the double dash when passing in -deterministic by APT

diff --git a/Engine/Plugins/Performance/AutomatedPerfTesting/Build/Scripts/AutomatedPerfTestNode.cs b/Engine/Plugins/Performance/AutomatedPerfTesting/Build/Scripts/AutomatedPerfTestNode.cs
--- a/Engine/Plugins/Performance/AutomatedPerfTesting/Build/Scripts/AutomatedPerfTestNode.cs
+++ b/Engine/Plugins/Performance/AutomatedPerfTesting/Build/Scripts/AutomatedPerfTestNode.cs
@@ -785,8 +785,11 @@ namespace AutomatedPerfTest
 			// the controller will be added by the subclasses
 
 			ClientRole.CommandLineParams.AddOrAppendParamValue("logcmds", "LogHttp Verbose, LogAutomatedPerfTest Verbose");
+
+			// force the screen to stay awake during the test
+			ClientRole.CommandLineParams.Add("keepscreenawake");
 			
-			ClientRole.CommandLineParams.Add("-deterministic");
+			ClientRole.CommandLineParams.Add("deterministic");
 
 			Log.Info("AutomatedPerfTestNode<>.GetConfiguration(): Config.DoFPSChart={0}, Config.DoCSVProfiler={1}, Config.DoVideoCapture={2}, Config.DoInsightsTrace={3}, Config.DoLLM={4}", 
 				Config.DoFPSChart, Config.DoCSVProfiler, Config.DoVideoCapture, Config.DoInsightsTrace, Config.DoLLM);
diff --git a/Engine/Source/Runtime/Engine/Private/UnrealEngine.cpp b/Engine/Source/Runtime/Engine/Private/UnrealEngine.cpp
--- a/Engine/Source/Runtime/Engine/Private/UnrealEngine.cpp
+++ b/Engine/Source/Runtime/Engine/Private/UnrealEngine.cpp
@@ -2159,6 +2159,12 @@ void UEngine::Init(IEngineLoop* InEngineLoop)
 	{
 		EnableScreenSaver( false );
 	}
+	
+	// enforce disabling screen saver when it's requested from commandline
+	if (FParse::Param(FCommandLine::Get(), TEXT("keepscreenawake")))
+	{
+		FPlatformApplicationMisc::ControlScreensaver(FPlatformApplicationMisc::EScreenSaverAction::Disable);
+	}
 
 	if (!IsRunningDedicatedServer() && !IsRunningCommandlet())
 	{
