Add Performance Insights to Artificats to Horde Results

diff --git a/Engine/Plugins/Performance/AutomatedPerfTesting/Build/Scripts/AutomatedPerfTestNode.cs b/Engine/Plugins/Performance/AutomatedPerfTesting/Build/Scripts/AutomatedPerfTestNode.cs
--- a/Engine/Plugins/Performance/AutomatedPerfTesting/Build/Scripts/AutomatedPerfTestNode.cs
+++ b/Engine/Plugins/Performance/AutomatedPerfTesting/Build/Scripts/AutomatedPerfTestNode.cs
@@ -465,7 +465,19 @@ namespace AutomatedPerfTest
 
 			if (Config.DoGPUReshape)
 			{
-				Report = ChainGPUReshapeReports(Config, Result, Report, OutputPath);
+				string ReportName = $"GPU Reshape - {Config.GPUReshapeWorkspace}";
+				string FileFilter = "*GRS.Report.html";
+
+				Report = AddTestReportArtificats(Result, Report, ReportName, OutputPath, FileFilter);
+			}
+
+			if (Config.DoInsightsTrace)
+			{
+				string UtraceLocation = Path.Combine(OutputPath, "Traces");
+				string FileFilter = "*.utrace";
+				string ArtifactReportName = "Trace Insights";
+
+				Report = AddTestReportArtificats(Result, Report, ArtifactReportName, UtraceLocation, FileFilter);
 			}
 
 			return Report;
@@ -898,25 +910,26 @@ namespace AutomatedPerfTest
 			}
 		}
 
-		private ITestReport ChainGPUReshapeReports(TConfigClass Config, TestResult Result, ITestReport Report, string OutputPath)
+		private ITestReport AddTestReportArtificats(TestResult Result, ITestReport Report, string TestName, string OutputPath, string FileFilter)
 		{
-			foreach (string GRSReportPath in Directory.GetFiles(OutputPath, "*GRS.Report.html"))
+			foreach (string AbsoluteFilePath in FindFiles(FileFilter, true, OutputPath))
 			{
-				// Create a typical report per, just keeps a link around
+				// Create a report per file found that successfully creates a Report of Type SimpleTestReport 
 				if (CreateSimpleReportForHorde(Result) is not HordeReport.SimpleTestReport ChildReport)
 				{
 					continue;
 				}
 
-				// Add to artifact
-				string BaseName = Path.GetFileName(GRSReportPath);
-				ChildReport.AttachArtifact(GRSReportPath, BaseName);
-				
+				string BaseFileName = Path.GetFileName(AbsoluteFilePath);
+
+				// Setup Report Info
+				ChildReport.AttachArtifact(AbsoluteFilePath, BaseFileName);
+				ChildReport.TestName = TestName;
+
 				// TODO: The general artifact linking is a work in progress, pending the V2 format
-				ChildReport.TestName = $"GPU Reshape - {Config.GPUReshapeWorkspace}";
-				ChildReport.URLLink  = $"/api/v2/artifacts/{Context.WorkerJobID}/file?path={BaseName}&inline=true";
-					
-				// Chain to last
+				ChildReport.URLLink = $"/api/v2/artifacts/{Context.WorkerJobID}/file?path={BaseFileName}&inline=true";
+
+				// Add Test Report Artificats to Report
 				switch (Report)
 				{
 					case HordeReport.BaseHordeReport BaseReport:
