Add Inisghts & Horde Job Link to PRS

diff --git a/Engine/Plugins/Performance/AutomatedPerfTesting/Build/Inc/AutomatedPerfTestCommonSettings.xml b/Engine/Plugins/Performance/AutomatedPerfTesting/Build/Inc/AutomatedPerfTestCommonSettings.xml
--- a/Engine/Plugins/Performance/AutomatedPerfTesting/Build/Inc/AutomatedPerfTestCommonSettings.xml
+++ b/Engine/Plugins/Performance/AutomatedPerfTesting/Build/Inc/AutomatedPerfTestCommonSettings.xml
@@ -1,5 +1,10 @@
 <?xml version="1.0" encoding="utf-8"?>
-<BuildGraph xmlns="http://www.epicgames.com/BuildGraph" xmlns:xsi="http://www.w3.org/2001/XMLSchema-instance" xsi:schemaLocation="http://www.epicgames.com/BuildGraph ../../../../../../../Engine/Build/Graph/Schema.xsd" >
+<BuildGraph xmlns="http://www.epicgames.com/BuildGraph" xmlns:xsi="http://www.w3.org/2001/XMLSchema-instance" xsi:schemaLocation="http://www.epicgames.com/BuildGraph ../../../../../../Engine/Build/Graph/Schema.xsd" >
+	
+	<!-- Horde Configuration -->
+	<EnvVar Name="UE_HORDE_JOBID"/>
+	<EnvVar Name="UE_HORDE_STEPID"/>
+	<Property Name="JobDetails" Value="$(HordeDashboardURL)/job/$(UE_HORDE_JOBID)?step=$(UE_HORDE_STEPID)"/>
 
 	<!-- Insights Setup for Automated Perf Testing -->
 	<!-- since we need to tell UAT whether or not to using the shipping configuration of Insights when processing -->
@@ -57,6 +62,7 @@
 	<Property Name="AutomatedPerformanceCommonArgs" Value="$(AutomatedPerformanceCommonArgs) -AutomatedPerfTest.DoFPSChart" If="$(DoFPSChart)"/>
 
 	<Property Name="AutomatedPerformanceCommonArgs" Value="$(AutomatedPerformanceCommonArgs) -PerfReportServer" If="$(PerfReportServer) And $(DoCSVProfiler)"/>
+	<Property Name="AutomatedPerformanceCommonArgs" Value="$(AutomatedPerformanceCommonArgs) -JobDetails=&quot;$(JobDetails)&quot;" If="'$(JobDetails)' != ''"/>
 	<Property Name="AutomatedPerformanceCommonArgs" Value="$(AutomatedPerformanceCommonArgs) -NoLocalReports" If="!$(GenerateLocalReports) And $(DoCSVProfiler)" />
 	<Property Name="AutomatedPerformanceCommonArgs" Value="$(AutomatedPerformanceCommonArgs) -resX=$(ResX) -resY=$(ResY)"/>
 	<Property Name="AutomatedPerformanceCommonArgs" Value="$(AutomatedPerformanceCommonArgs) -AutomatedPerfTest.DoVideoCapture" If="$(DoVideoCapture)"/>
diff --git a/Engine/Plugins/Performance/AutomatedPerfTesting/Build/Scripts/AutomatedPerfTestNode.cs b/Engine/Plugins/Performance/AutomatedPerfTesting/Build/Scripts/AutomatedPerfTestNode.cs
--- a/Engine/Plugins/Performance/AutomatedPerfTesting/Build/Scripts/AutomatedPerfTestNode.cs
+++ b/Engine/Plugins/Performance/AutomatedPerfTesting/Build/Scripts/AutomatedPerfTestNode.cs
@@ -91,6 +91,11 @@ namespace AutomatedPerfTest
 		/// </summary>
 		public string BaseOutputPath { get; protected set; }
 
+		/// <summary>
+		/// The Horde job link.
+		/// </summary>
+		public string HordeJobLink { get; protected set; }
+
 		/// <summary>
 		/// Returns the APT Performance Artifact Path. 
 		/// </summary>
@@ -328,7 +333,8 @@ namespace AutomatedPerfTest
 			UnrealTargetPlatform Platform = Context.GetRoleContext(UnrealTargetRole.Client).Platform;
 			TConfigClass Config = GetCachedConfiguration();
 			string OutputPath = GetPerformanceReportArtifactOutputPath(Platform);
-			
+			HordeJobLink = Globals.Params.ParseValue("JobDetails", "");
+
 			// Always render reshape reports, failed tests may contain validation of interest
 			if (Config.DoGPUReshape)
 			{
@@ -367,7 +373,7 @@ namespace AutomatedPerfTest
 					{
 						Dictionary<string, dynamic> CommonDataSourceFields = new Dictionary<string, dynamic>
 						{
-							{"HordeJobUrl", Globals.Params.ParseValue("JobDetails", null)}
+							{"HordeJobUrl", HordeJobLink}
 						};
 						
 						Log.Info("Creating perf server importer with build name {BuildName}", OriginalBuildName);
