Update PF PRS warning line into a report output without warnings and for all runs.

Also cleaned up Report passing/returning since it's an object and CS treats all objects as passed by reference by default.

#DisabledValidations Preflight

diff --git a/Engine/Plugins/Performance/AutomatedPerfTesting/Build/Scripts/AutomatedPerfTestNode.cs b/Engine/Plugins/Performance/AutomatedPerfTesting/Build/Scripts/AutomatedPerfTestNode.cs
--- a/Engine/Plugins/Performance/AutomatedPerfTesting/Build/Scripts/AutomatedPerfTestNode.cs
+++ b/Engine/Plugins/Performance/AutomatedPerfTesting/Build/Scripts/AutomatedPerfTestNode.cs
@@ -334,6 +334,8 @@ namespace AutomatedPerfTest
 			TConfigClass Config = GetCachedConfiguration();
 			string OutputPath = GetPerformanceReportArtifactOutputPath(Platform);
 			HordeJobLink = Globals.Params.ParseValue("JobDetails", "");
+			string DataSourceName = GetConfiguration().DataSourceName;
+			string PRSReportLink = "";
 
 			// Always render reshape reports, failed tests may contain validation of interest
 			if (Config.DoGPUReshape)
@@ -377,7 +379,7 @@ namespace AutomatedPerfTest
 						};
 						
 						Log.Info("Creating perf server importer with build name {BuildName}", OriginalBuildName);
-						string DataSourceName = GetConfiguration().DataSourceName;
+						
 						string ImportDirOverride = Globals.Params.ParseValue("PerfReportServerImportDir", null);
 						Log.Info("Creating PRS Importer for data source '{0}' and import dir override (if any) '{1}'.",
 							DataSourceName, ImportDirOverride);
@@ -440,7 +442,8 @@ namespace AutomatedPerfTest
 							}
 
 							// Create the import batch
-							Importer.Import(CsvImportEntries);
+							Importer.Import(CsvImportEntries, out PRSReportLink);
+
 							// trust blindly for now...
 							bGeneratedPRSReport = true;
 						}
@@ -474,16 +477,18 @@ namespace AutomatedPerfTest
 				string ReportName = $"GPU Reshape - {Config.GPUReshapeWorkspace}";
 				string FileFilter = "*GRS.Report.html";
 
-				Report = AddTestReportArtificats(Result, Report, ReportName, OutputPath, FileFilter);
+				AddTestReportArtifacts(Result, Report, ReportName, OutputPath, FileFilter);
 			}
 
 			if (Config.DoInsightsTrace)
 			{
 				string UtraceLocation = Path.Combine(OutputPath, "Traces");
 				string FileFilter = "*.utrace";
-				string ArtifactReportName = "Trace Insights";
+				string ArtifactReportName = "Trace Insights Report";
+				string PRSReportName = "PRS Report Upload";
 
-				Report = AddTestReportArtificats(Result, Report, ArtifactReportName, UtraceLocation, FileFilter);
+				AddTestReportArtifacts(Result, Report, ArtifactReportName, UtraceLocation, FileFilter);
+				AddLinkToReport(Result, Report, PRSReportName, PRSReportLink);
 			}
 
 			return Report;
@@ -916,7 +921,7 @@ namespace AutomatedPerfTest
 			}
 		}
 
-		private ITestReport AddTestReportArtificats(TestResult Result, ITestReport Report, string TestName, string OutputPath, string FileFilter)
+		private void AddTestReportArtifacts(TestResult Result, ITestReport Report, string TestName, string OutputPath, string FileFilter)
 		{
 			foreach (string AbsoluteFilePath in FindFiles(FileFilter, true, OutputPath))
 			{
@@ -932,24 +937,13 @@ namespace AutomatedPerfTest
 				ChildReport.AttachArtifact(AbsoluteFilePath, BaseFileName);
 				ChildReport.TestName = TestName;
 
-				// TODO: The general artifact linking is a work in progress, pending the V2 format
-				ChildReport.URLLink = $"/api/v2/artifacts/{Context.WorkerJobID}/file?path={BaseFileName}&inline=true";
+				// TODO: UE-305553
+				ChildReport.URLLink = $"/api/v2/artifacts/{Context.WorkerJobID}/file?path={BaseFileName}&inline=true";				
 
-				// Add Test Report Artificats to Report
-				switch (Report)
-				{
-					case HordeReport.BaseHordeReport BaseReport:
-						BaseReport.AttachDependencyReport(ChildReport, ChildReport.TestName);
-						break;
-					default:
-						Report = ChildReport;
-						break;
-				}
+				AttachHordeReport(Report, ChildReport);
 			}
-
-			return Report;
 		}
-		
+
 		public void CopyInsightsTraceToOutput(string FromArtifactPath, string ToOutputPath, out string OutTraceName)
 		{
 			OutTraceName = string.Empty;
diff --git a/Engine/Source/Programs/AutomationTool/Gauntlet/Framework/Utils/Gauntlet.CsvImporter.cs b/Engine/Source/Programs/AutomationTool/Gauntlet/Framework/Utils/Gauntlet.CsvImporter.cs
--- a/Engine/Source/Programs/AutomationTool/Gauntlet/Framework/Utils/Gauntlet.CsvImporter.cs
+++ b/Engine/Source/Programs/AutomationTool/Gauntlet/Framework/Utils/Gauntlet.CsvImporter.cs
@@ -200,6 +200,7 @@ namespace Gauntlet
 	public interface ICsvImporter
 	{
 		public void Import(IEnumerable<ICsvImport> Imports);
+		public void Import(IEnumerable<ICsvImport> Imports, out string URLLink);
 	}
 
 	/// <summary>
diff --git a/Engine/Source/Programs/AutomationTool/Gauntlet/Unreal/Base/Gauntlet.UnrealTestNode.cs b/Engine/Source/Programs/AutomationTool/Gauntlet/Unreal/Base/Gauntlet.UnrealTestNode.cs
--- a/Engine/Source/Programs/AutomationTool/Gauntlet/Unreal/Base/Gauntlet.UnrealTestNode.cs
+++ b/Engine/Source/Programs/AutomationTool/Gauntlet/Unreal/Base/Gauntlet.UnrealTestNode.cs
@@ -831,6 +831,52 @@ namespace Gauntlet
 			return true;
 		}
 
+		/// <summary>
+		/// Used to Add a URL to a Horde Child Report.
+		/// </summary>
+		/// <param name="Result">The current TestResult.</param>
+		/// <param name="Report">The curret ITestReport.</param>
+		/// <param name="ReportName">The name you want to show up on the Horde Child Report.</param>
+		/// <param name="URLLink">The URL Link you want to populate on the Horde Child Report.</param>
+		public void AddLinkToReport(TestResult Result, ITestReport Report, string ReportName, string URLLink)
+		{
+			if (CreateSimpleReportForHorde(Result) is not HordeReport.SimpleTestReport ChildReport)
+			{
+				return;
+			}
+
+			if(!string.IsNullOrEmpty(ReportName) && !string.IsNullOrEmpty(URLLink))
+			{
+				Log.Info($"Adding a Horde Child Report URL Link for Report Name: {ReportName}, URLLink: {URLLink}.");
+				ChildReport.TestName = ReportName;
+				ChildReport.URLLink = URLLink;
+				AttachHordeReport(Report, ChildReport);
+
+				return;
+			}
+
+			Log.Warning("A Horde Child Report URL Link for Report Name: {ReportName}, URLLink: {URLLink} was requested but could not be generated.");
+
+		}
+
+		/// <summary>
+		/// Used to Attach a Horde Report to a report whether it's a base report, child, or another type.
+		/// </summary>
+		/// <param name="Report">The base report to attach the child report to.</param>
+		/// <param name="ChildReport">The child report to attach to the base report.</param>
+		public void AttachHordeReport(ITestReport Report, HordeReport.SimpleTestReport ChildReport)
+		{
+			switch (Report)
+			{
+				case HordeReport.BaseHordeReport BaseReport:
+					BaseReport.AttachDependencyReport(ChildReport, ChildReport.TestName);
+					break;
+				default:
+					Report = ChildReport;
+					break;
+			}
+		}
+
 		public virtual void PopulateCommandlineInfo()
 		{
 			SupportedParameters.Add(new GauntletParamDescription()
