APT: Add escaped branch name to build version when running tests via Automated Perf Test plugin

#rnx

diff --git a/Engine/Plugins/Performance/AutomatedPerfTesting/Build/Inc/AutomatedPerfTestCommonSettings.xml b/Engine/Plugins/Performance/AutomatedPerfTesting/Build/Inc/AutomatedPerfTestCommonSettings.xml
--- a/Engine/Plugins/Performance/AutomatedPerfTesting/Build/Inc/AutomatedPerfTestCommonSettings.xml
+++ b/Engine/Plugins/Performance/AutomatedPerfTesting/Build/Inc/AutomatedPerfTestCommonSettings.xml
@@ -5,7 +5,7 @@
 	<EnvVar Name="UE_HORDE_JOBID"/>
 	<EnvVar Name="UE_HORDE_STEPID"/>
 	<Property Name="JobDetails" Value="$(HordeDashboardURL)/job/$(UE_HORDE_JOBID)?step=$(UE_HORDE_STEPID)"/>
-
+	<Property Name="APTBuildVersion" Value="$(EscapedBranch)-CL-$(Change)-$(ProjectName)$(PreflightSuffix)"/> 
 	<!-- Insights Setup for Automated Perf Testing -->
 	<!-- since we need to tell UAT whether or not to using the shipping configuration of Insights when processing -->
 	<!-- insights traces after the conclusion of the test --> 
@@ -17,7 +17,7 @@
 	<Property Name="UnrealInsightsBuildConfiguration" Value="Shipping" If="$(UseShippingInsights)" />
 	
 	<!-- Other common arguments --> 
-	<Option Name="TestID" DefaultValue="" Description="A string to use to identify these performance tests" />
+	<Option Name="TestID" DefaultValue="APT" Description="A string to use to identify these performance tests" />
 	<Option Name="DeviceProfileOverride" DefaultValue="" Description="Which device profile to use for this test instead of the default value. Will be prepended by PlatformName." />
 	<Option Name="DoCSVProfiler" DefaultValue="true" Restrict="true|false" Description="Whether or not to use CSVProfiler to gather and parse performance telemetry." />
 	<Option Name="PerfReportServer" DefaultValue="false" Restrict="true|false" Description="Whether or not to use CSVProfiler to gather and parse performance telemetry." />
@@ -37,9 +37,12 @@
 	<Option Name="APTIgnoreTestBuildLogging" Restrict="true|false" DefaultValue="true" Description="If true, tests run in Test build configuration will not fail if logs are empty." />
 	<Option Name="PerfOutputPath" DefaultValue="" Description="Path where Perf Test artifacts (Logs, CSVs, Reports, Cache) will be stored. Override to a common network-accessible path when running on build machines if needed"/>
 	<Option Name="ProjectPath" DefaultValue="$(DefaultProjectPath)" Restrict=".+" Description="Path to the folder that contains your project"/>
-	
+	<Option Name="TestBuildVersion" DefaultValue="$(APTBuildVersion)" Description="Overrides test build version"/>
 	<Option Name="ExtraAutomatedPerformanceCommonArgs" DefaultValue="" Description="Additional Common arguments passed to Gauntlet command line"/>
 
+	<!-- Override build version with given value. If this value is empty, APT will use default build version-->
+	<Property Name="BuildVersion" Value="$(TestBuildVersion)"/> 
+
 	<!-- Plugin-provided tests -->
 	<Option Name="WithAPTSequenceTests" DefaultValue="true" Restrict="true|false" Description="Whether to run Auto Perf Test sequence flythrough test." />
 	<Option Name="WithAPTStaticCameraTests" DefaultValue="false" Restrict="true|false" Description="Whether to run Auto Perf Test static camera tests."  />
@@ -69,6 +72,8 @@
 	<Property Name="AutomatedPerformanceCommonArgs" Value="$(AutomatedPerformanceCommonArgs) -AutomatedPerfTest.IgnoreTestBuildLogging" If="$(APTIgnoreTestBuildLogging)"/>
 	<Property Name="AutomatedPerformanceCommonArgs" Value="$(AutomatedPerformanceCommonArgs) -AutomatedPerfTest.PerfOutputPath=&quot;$(PerfOutputPath)&quot;" If="'$(PerfOutputPath)' != ''"/>
 	<Property Name="AutomatedPerformanceCommonArgs" Value="$(AutomatedPerformanceCommonArgs) $(ExtraAutomatedPerformanceCommonArgs)" If="'$(ExtraAutomatedPerformanceCommonArgs)' != ''"/>
+	<Property Name="AutomatedPerformanceCommonArgs" Value="$(AutomatedPerformanceCommonArgs) -Branch=&quot;$(EscapedBranch)&quot; -CL=$(Change)"/>
+	<Property Name="AutomatedPerformanceCommonArgs" Value="$(AutomatedPerformanceCommonArgs) -AutomatedPerfTest.TestBuildVersion=&quot;$(BuildVersion)&quot;"/>
 
 	<Property Name="PerfTestCommonArgs"     Value="$(AutomatedPerformanceCommonArgs) -iterations=$(IterationsPerf)"/>
 	<Property Name="LLMTestCommonArgs"      Value="$(AutomatedPerformanceCommonArgs) -iterations=$(IterationsLLM)"/>
diff --git a/Engine/Plugins/Performance/AutomatedPerfTesting/Build/Scripts/AutomatedPerfTestConfig.cs b/Engine/Plugins/Performance/AutomatedPerfTesting/Build/Scripts/AutomatedPerfTestConfig.cs
--- a/Engine/Plugins/Performance/AutomatedPerfTesting/Build/Scripts/AutomatedPerfTestConfig.cs
+++ b/Engine/Plugins/Performance/AutomatedPerfTesting/Build/Scripts/AutomatedPerfTestConfig.cs
@@ -161,6 +161,12 @@ namespace AutomatedPerfTest
 		[AutoParamWithNames("", "AutomatedPerfTest.TestConfigName")]
 		public string TestConfigName = "";
 
+		/// <summary>
+		/// Overrides test build version. If empty, default build version will be used. 
+		/// </summary>
+		[AutoParamWithNames("", "AutomatedPerfTest.TestBuildVersion")]
+		public string TestBuildVersion = "";
+
 		public string DataSourceName;
 
 		/// <summary>
diff --git a/Engine/Plugins/Performance/AutomatedPerfTesting/Build/Scripts/AutomatedPerfTestNode.cs b/Engine/Plugins/Performance/AutomatedPerfTesting/Build/Scripts/AutomatedPerfTestNode.cs
--- a/Engine/Plugins/Performance/AutomatedPerfTesting/Build/Scripts/AutomatedPerfTestNode.cs
+++ b/Engine/Plugins/Performance/AutomatedPerfTesting/Build/Scripts/AutomatedPerfTestNode.cs
@@ -1012,6 +1012,21 @@ namespace AutomatedPerfTest
 			ClientRole.CommandLineParams.Add("AutomatedPerfTest.DoCSVProfiler");
 			ClientRole.CommandLineParams.Add("csvGpuStats");
 
+			if(string.IsNullOrEmpty(Config.TestBuildVersion))
+			{
+				Config.TestBuildVersion = OriginalBuildName;
+			}
+			else
+			{
+				// If build version has been overridden for any reason,
+				// this ensures that the CSV metadata matches the test
+				// build version since this var is used while instantiating
+				// perf report importer. 
+				OriginalBuildName = Config.TestBuildVersion;
+			}
+
+			Log.Info($"Using Test Build Version {Config.TestBuildVersion}");
+
 			// Add CSV metadata
 			List<string> CsvMetadata =
 			[
@@ -1019,7 +1034,7 @@ namespace AutomatedPerfTest
 				"gauntletTestType=AutomatedPerfTest",
 				$"gauntletSubTest={GetSubtestName(Config)}",
 				"testBuildIsPreflight=" + (ReportGenUtils.IsTestingPreflightBuild(OriginalBuildName) ? "1" : "0"),
-				$"testBuildVersion={OriginalBuildName}",
+				$"testBuildVersion={Config.TestBuildVersion}",
 				"testconfigname=" + Config.TestConfigName,
 			];
 
