Add AutomatedPerfTest.NullOSS parameter to APT to allow for overriding the online subsystem with the null implementation

diff --git a/Engine/Plugins/Performance/AutomatedPerfTesting/Build/Scripts/AutomatedPerfTestConfig.cs b/Engine/Plugins/Performance/AutomatedPerfTesting/Build/Scripts/AutomatedPerfTestConfig.cs
--- a/Engine/Plugins/Performance/AutomatedPerfTesting/Build/Scripts/AutomatedPerfTestConfig.cs
+++ b/Engine/Plugins/Performance/AutomatedPerfTesting/Build/Scripts/AutomatedPerfTestConfig.cs
@@ -64,6 +64,18 @@ namespace AutomatedPerfTest
 		[AutoParamWithNames(false, "AutomatedPerfTest.DoFPSChart")]
 		public bool DoFPSChart;
 
+		/// <summary>
+		/// Whether to wait for a debugger to attach to an APT launched test
+		/// </summary>
+		[AutoParamWithNames(false, "AutomatedPerfTest.WaitForAttach")]
+		public bool WaitForAttach;
+
+		/// <summary>
+		/// Whether to use the null online subsystem implementation
+		/// </summary>
+		[AutoParamWithNames(false, "AutomatedPerfTest.NullOSS")]
+		public bool NullOSS;
+
 		/// <summary>
 		/// Whether to run the test with -llm
 		/// </summary>
diff --git a/Engine/Plugins/Performance/AutomatedPerfTesting/Build/Scripts/AutomatedPerfTestNode.cs b/Engine/Plugins/Performance/AutomatedPerfTesting/Build/Scripts/AutomatedPerfTestNode.cs
--- a/Engine/Plugins/Performance/AutomatedPerfTesting/Build/Scripts/AutomatedPerfTestNode.cs
+++ b/Engine/Plugins/Performance/AutomatedPerfTesting/Build/Scripts/AutomatedPerfTestNode.cs
@@ -839,6 +839,16 @@ namespace AutomatedPerfTest
 				ClientRole.CommandLineParams.Add("llmcsv");
 			}
 
+			if (Config.WaitForAttach)
+			{
+				ClientRole.CommandLineParams.Add("waitforattach");
+			}
+
+			if (Config.NullOSS)
+			{
+				ClientRole.CommandLineParams.Add("ini:Engine:[OnlineSubsystem]:DefaultPlatformService=Null");
+			}
+
 			if (Config.DoGPUPerf)
 			{
 				// see ReplayRun.py, reducedAsyncComputeCommands
