Fix Lyra broken due to android automation changes

#rnx

diff --git a/Engine/Plugins/Performance/AutomatedPerfTesting/Build/Scripts/AutomatedPerfTestNode.cs b/Engine/Plugins/Performance/AutomatedPerfTesting/Build/Scripts/AutomatedPerfTestNode.cs
--- a/Engine/Plugins/Performance/AutomatedPerfTesting/Build/Scripts/AutomatedPerfTestNode.cs
+++ b/Engine/Plugins/Performance/AutomatedPerfTesting/Build/Scripts/AutomatedPerfTestNode.cs
@@ -1214,9 +1214,16 @@ namespace AutomatedPerfTest
 			{
 				foreach (UnrealTestRole ClientRole in Config.GetRequiredRoles(UnrealTargetRole.Client))
 				{
-					UnrealTargetPlatform Platform = Context.GetRoleContext(UnrealTargetRole.Client).Platform;
+					UnrealTestRoleContext RoleContext = Context.GetRoleContext(UnrealTargetRole.Client);
+					UnrealTargetPlatform Platform = RoleContext.Platform;
 					if (Device.Platform == Platform && ClientRole.FilesToCopy.Any())
 					{
+						UnrealSessionRole Role = new(RoleContext.Type, Platform, RoleContext.Configuration);
+						UnrealAppConfig AppConfig = Context.BuildInfo.CreateConfiguration(Role);
+
+						// needed so AFS will work on CopyAdditional files
+						_ = Device.CreateAppInstall(AppConfig);
+
 						Device.CopyAdditionalFiles(ClientRole.FilesToCopy);
 					}
 				}
diff --git a/Engine/Source/Programs/AutomationTool/Gauntlet/Platform/Android/Gauntlet.TargetDeviceAndroid.cs b/Engine/Source/Programs/AutomationTool/Gauntlet/Platform/Android/Gauntlet.TargetDeviceAndroid.cs
--- a/Engine/Source/Programs/AutomationTool/Gauntlet/Platform/Android/Gauntlet.TargetDeviceAndroid.cs
+++ b/Engine/Source/Programs/AutomationTool/Gauntlet/Platform/Android/Gauntlet.TargetDeviceAndroid.cs
@@ -433,7 +433,15 @@ namespace Gauntlet
 					FileInfo SourceFile = new FileInfo(FileToCopy.SourceFileLocation);
 					SourceFile.IsReadOnly = false;
 
-					CopyFileToDevice(InstallCache.AndroidPackageName, SourceFile.FullName, DestinationFile, InstallCache.ProjectFile, InstallCache.Configuration);
+					// InstallCache may be null so I'll add back the previous call in that case
+					if (InstallCache != null)
+					{
+						CopyFileToDevice(InstallCache.AndroidPackageName, SourceFile.FullName, DestinationFile, InstallCache.ProjectFile, InstallCache.Configuration);
+					}
+					else
+					{
+						CopyFileToDevice(null, SourceFile.FullName, DestinationFile);
+					}
 				}
 				else
 				{
