APT: Add capability to auto-activate configs for certain roles.

* Add config "bridging" capability such that config bridges with the right attributes are activated only when certain params are passed.
* This allows for ease of specifying and composing repeatable test params and/or for new test types which can then be activated based on params.
* Fix ProfileGo settings update bug.

#rnx

diff --git a/Engine/Plugins/Developer/ProjectLauncher/Source/CommonLaunchExtensions/Private/AutomatedPerfTest/AutomatedPerfTestLaunchExtension.cpp b/Engine/Plugins/Developer/ProjectLauncher/Source/CommonLaunchExtensions/Private/AutomatedPerfTest/AutomatedPerfTestLaunchExtension.cpp
--- a/Engine/Plugins/Developer/ProjectLauncher/Source/CommonLaunchExtensions/Private/AutomatedPerfTest/AutomatedPerfTestLaunchExtension.cpp
+++ b/Engine/Plugins/Developer/ProjectLauncher/Source/CommonLaunchExtensions/Private/AutomatedPerfTest/AutomatedPerfTestLaunchExtension.cpp
@@ -426,12 +426,35 @@ void FAutomatedPerfTestLaunchExtensionInstance::AddProfileGoTestNodeOptions(Proj
 		.Text(LOCTEXT("TestExportLabelButton", "Export JSON"))
 		.OnClicked_Lambda([this]() -> FReply
 		{
+			// TEMP: Will be removed once ProfileGo scenario workflow creation is finalized.
 			const TCHAR* ProfileGoFilename = TEXT("ProfileGo.json");
 			ExportProfileGoScenarios(ProfileGoFilename);
 			return FReply::Handled();
 		})
 		.IsEnabled_Lambda([this]() { return GetTestType() == EAutomatedPerfTestType::ProfileGo && IsEditorBuild(); })
 	)
+	.AddWidget(LOCTEXT("SettingsUpdateLabel", "Update ProfileGo Settings"),
+		MoveTemp(Callbacks),
+		SNew(SButton)
+		.HAlign(HAlign_Center)
+		.VAlign(VAlign_Center)
+		.Text(LOCTEXT("SettingsUpdateLabelButton", "Update"))
+		.OnClicked_Lambda([this]() -> FReply
+		{
+		#if WITH_EDITOR
+			// TEMP: Will be removed once ProfileGo scenario workflow creation is finalized.
+			UWorld* World = GWorld;
+			UEditorEngine* EditorEngine = Cast<UEditorEngine>(GEngine);
+			if (GIsEditor && EditorEngine != nullptr && EditorEngine->PlayWorld != nullptr)
+			{
+				World = EditorEngine->PlayWorld.Get();
+			}
+			UProfileGo::GetCDO().AddScenariosInLevel(World);
+		#endif
+			return FReply::Handled();
+		})
+		.IsEnabled_Lambda([this]() { return GetTestType() == EAutomatedPerfTestType::ProfileGo && IsEditorBuild(); })
+	)
 	.AddFileString(LOCTEXT("ProfileGoConfigLabel","Config File"),
 	{
 		.GetValue = [this]() { return GetConfigString(EConfig::PerProfile, ConfigFileName); },
diff --git a/Engine/Plugins/Performance/AutomatedPerfTesting/Build/Scripts/AutomatedPerfTest.Bridge.cs b/Engine/Plugins/Performance/AutomatedPerfTesting/Build/Scripts/AutomatedPerfTest.Bridge.cs
new file mode 100644
--- /dev/null
+++ b/Engine/Plugins/Performance/AutomatedPerfTesting/Build/Scripts/AutomatedPerfTest.Bridge.cs
@@ -0,0 +1,294 @@
+// Copyright Epic Games, Inc. All Rights Reserved.
+
+using AutomatedPerfTesting;
+using Gauntlet;
+using System;
+using System.Collections.Generic;
+using System.Data;
+using System.Linq;
+using System.Reflection;
+
+/// <summary>
+/// This file defines the Test Bridge interface. In this context, a "Test Bridge"
+/// refers to the custom test session configuration and logic which needs to run
+/// before Gauntlet launches the test session on the target device. 
+/// 
+/// A Test Bridge may contain various attributes which enable composing and defining 
+/// test-specific config and logic depending on the attribute parameters. These 
+/// bridges are then instantiated for a given test if the test context parameters match
+/// the attribute parameters. 
+/// 
+/// Simple example would be defining Test Bridges for each Test Type in your project
+/// i.e. SequenceBridge, ReplayBridge, ProfileGoBridge etc. such that each bridge
+/// is only activated for their corresponding test type. 
+/// </summary>
+namespace AutomatedPerfTest
+{
+	/// <summary>
+	/// Interface to be used if a project wants to configure the Automated Perf Test
+	/// context, roles and/or config without making intrusive changes to the main 
+	/// node before launching the test. 
+	/// 
+	/// Ideally project/game specific command line parameters get set here. 
+	/// </summary>
+	public interface ITestBridge
+	{
+		/// <summary>
+		/// Initial configure call. This is invoked first.
+		/// </summary>
+		/// <param name="Context">Test Context</param>
+		/// <param name="Config">APT Config Base</param>
+		void Configure(UnrealTestContext Context, AutomatedPerfTestConfigBase Config) { }
+
+		/// <summary>
+		/// Configure client role
+		/// </summary>
+		/// <param name="Role">Client Test Role</param>
+		/// <param name="Context">Test Context</param>
+		/// <param name="Config">APT Config Base</param>
+		void ConfigureClient(UnrealTestRole Role, UnrealTestContext Context, AutomatedPerfTestConfigBase Config) { }
+
+		/// <summary>
+		/// Configure server role. This is invoked optionally if a server role is required for the test. 
+		/// </summary>
+		/// <param name="Role">Server Test Role</param>
+		/// <param name="Context">Test Context</param>
+		/// <param name="Config">APT Config Base</param>
+		void ConfigureServer(UnrealTestRole Role, UnrealTestContext Context, AutomatedPerfTestConfigBase Config) { }
+
+		/// <summary>
+		/// If this is true, server role is instantiated and configured. 
+		/// </summary>
+		/// <returns></returns>
+		bool IsServerRoleRequired() { return false; }
+	}
+
+	/// <summary>
+	/// Null Test Bridge
+	/// </summary>
+	public class NullBridge : ITestBridge { }
+
+
+	/// <summary>
+	/// Static helper class to activate and invoke test bridge configuration
+	/// functions.
+	/// </summary>
+	public static class AutoTestBridge
+	{
+		private static HashSet<ITestBridge> Bridges = new HashSet<ITestBridge>();
+
+		/// <summary>
+		/// Initialize Test Bridge instances with given list of Bridge class names. 
+		/// </summary>
+		/// <param name="InConfigBridges">Bridge names</param>
+		public static void Initialize(IEnumerable<string> InConfigBridges)
+		{
+			foreach (string BridgeName in InConfigBridges)
+			{
+				ITestBridge ConfigBridge = GetConfigBridge(BridgeName.Trim());
+				AddBridge(ConfigBridge);
+			}
+		}
+
+		public static void AddBridge(ITestBridge ConfigBridge)
+		{
+			if(ConfigBridge != null)
+			{
+				Bridges.Add(ConfigBridge);
+				Log.Info($"Test Bridge Added: {ConfigBridge.GetType().Name}");
+			}
+		}
+
+		private static ITestBridge GetConfigBridge(string ConfigBridgeName)
+		{
+			if (!string.IsNullOrEmpty(ConfigBridgeName))
+			{
+				Type BridgeType = Util.GetTypeWithInterface<ITestBridge>(ConfigBridgeName);
+				if (BridgeType != null)
+				{
+					return Activator.CreateInstance(BridgeType) as ITestBridge;
+				}
+			}
+
+			return null;
+		}
+
+		public static void Configure(UnrealTestContext Context, AutomatedPerfTestConfigBase Config)
+		{
+			Bridges.ToList().ForEach(Bridge => Bridge?.Configure(Context, Config));
+		}
+
+		public static void ConfigureClient(UnrealTestRole Role, UnrealTestContext Context, AutomatedPerfTestConfigBase Config)
+		{
+			Bridges.ToList().ForEach(Bridge => Bridge?.ConfigureClient(Role, Context, Config));
+		}
+
+		public static void ConfigureServer(UnrealTestRole Role, UnrealTestContext Context, AutomatedPerfTestConfigBase Config)
+		{
+			Bridges.ToList().ForEach(Bridge => Bridge?.ConfigureServer(Role, Context, Config));
+		}
+
+		public static bool IsServerRoleRequired()
+		{
+			return Bridges.Where(B => B.IsServerRoleRequired()).Any();
+		}
+
+		private static bool ValidateAttributes(Type TestBridge, string ProjectName = "")
+		{
+			EnableForProjectAttribute ProjectAttribute = TestBridge.GetCustomAttribute<EnableForProjectAttribute>();
+			EnableForParamsAttribute Params = TestBridge.GetCustomAttribute<EnableForParamsAttribute>();
+
+			AutoBridge BridgeAttribute = TestBridge.GetCustomAttribute<AutoBridge>();
+			IEnumerable<Attribute> AutoBridgeTestAttribute = TestBridge.GetCustomAttributes(typeof(AutoBridgeForTestAttribute<>));
+
+			bool Result = true;
+			if (ProjectAttribute != null)
+			{
+				Result = ProjectAttribute.ProjectName == ProjectName;
+			}
+
+			if (Params != null)
+			{
+				foreach (KeyValuePair<string, string> Args in Params.Params)
+				{
+					string Value = Globals.Params.ParseValue(Args.Key, "");
+					bool bContainsKey = Globals.Params.ParseParam(Args.Key);
+
+					// Some params may be bool params and may only have the param key to 
+					// indicate if it should be activated. 
+					bool bIsBoolVal = string.IsNullOrEmpty(Value) && string.IsNullOrEmpty(Args.Value) && bContainsKey;
+					bool bIsEqual = string.Equals(Value, Args.Value, StringComparison.OrdinalIgnoreCase);
+					if (!bIsBoolVal && !bIsEqual)
+					{
+						Result = false;
+						break;
+					}
+				}
+			}
+
+			return Result && (BridgeAttribute != null || AutoBridgeTestAttribute.Any());
+		}
+
+		/// <summary>
+		/// Registers valid Test Bridge types which will then be used to configure the test 
+		/// depending on the need of the test and/or project. 
+		/// Following attributes used for validation: <br/>
+		/// 
+		/// [AutoBridge] - This enables the bridge for all automated perf tests.<br/>
+		/// [AutoBridgeForTest[<paramref name="CurrentTestType"/>]] - This enables the bridge only 
+		/// for the specified test type. Multiple types are supported. <br/>
+		/// [EnableForProject] - Ensures the bridge is only enabled for a given project. <br/>
+		/// [EnableForParams]
+		/// 
+		/// </summary>
+		/// <param name="CurrentTestType">Type of test being run in this given session</param>
+		/// <param name="ProjectName">Name of project. This should match the name in 
+		/// EnableForProjectAttribute to take effect.</param>
+		public static void ActivateTestBridges(Type CurrentTestType, string ProjectName = "")
+		{
+			Type[] TestBridgeTypes = Util.GetTypesWithInterface<ITestBridge>();
+			foreach(Type TestBridge in TestBridgeTypes)
+			{
+				AutoBridge BridgeAttribute = TestBridge.GetCustomAttribute<AutoBridge>();
+				IEnumerable<Attribute> AutoBridgeTestAttribute = TestBridge.GetCustomAttributes(typeof(AutoBridgeForTestAttribute<>));
+
+				// A bridge is valid if each present attribute is valid. If an attribute is null, we do not consider
+				// it for validation. 
+				if (!ValidateAttributes(TestBridge, ProjectName))
+				{
+					continue;
+				}
+
+				if (BridgeAttribute != null)
+				{
+					// If a Test Bridge is enabled for all test types, we simply add the bridge
+					// and continue to the next bridge derived type. 
+					AddBridge(Activator.CreateInstance(TestBridge) as ITestBridge);
+					continue;
+				}
+
+				// Register bridges which have indicated that they are applicable for a given type of
+				// test. 
+				foreach (Attribute AutoBridge in AutoBridgeTestAttribute)
+				{
+					if (AutoBridge != null && AutoBridge.GetType().IsGenericType) 
+					{
+						Type CandidateTestType = AutoBridge.GetType().GetGenericArguments().First();
+						if(CandidateTestType != null && CurrentTestType == CandidateTestType)
+						{
+							AddBridge(Activator.CreateInstance(TestBridge) as ITestBridge);
+						}
+					}
+				}
+			}
+		}
+	}
+
+	/// <summary>
+	/// Automatically activate given Test Bridge for all Automated Perf Tests.
+	/// If this is the only attribute present, it will be added to all
+	/// tests by default. <br/>
+	/// Example: <code>[AutoBridge]</code>
+	/// </summary>
+	[AttributeUsage(AttributeTargets.Class, AllowMultiple = false)]
+	public class AutoBridge : Attribute
+	{
+		public AutoBridge() { }
+	}
+
+	/// <summary>
+	/// Automatically activate given Test Bridge for given Automated Perf Test
+	/// type. If this is the only attribute present, it will be added to
+	/// all instances of a given test type.<br/>
+	/// Example: <code>[AutoBridgeForTest[ReplayTest]]</code>
+	/// </summary>
+	[AttributeUsage(AttributeTargets.Class, AllowMultiple = true)]
+	public class AutoBridgeForTestAttribute<T> : Attribute
+		where T : IAutomatedPerfTest
+	{
+		public Type TestType { get; private set; }
+
+		public AutoBridgeForTestAttribute()
+		{ 
+			TestType = typeof(T);
+		}
+	}
+
+	/// <summary>
+	/// Ensures test bridge is activated only for a given project. This ensures if
+	/// a given bridge is very specific and/or does not apply to any other project
+	/// it can only be activated for that project. <br/> 
+	/// This should be combined with one of the AutoBridge attribute types to ensure
+	/// auto-activation of the bridge. <br/>
+	/// Example: <code>[EnableForProject("MyProjectName")]</code>
+	/// </summary>
+	[AttributeUsage(AttributeTargets.Class, AllowMultiple = true)]
+	public class EnableForProjectAttribute : Attribute
+	{
+		public string ProjectName { get; private set; }
+
+		public EnableForProjectAttribute(string InProjectName)
+		{
+			ProjectName = InProjectName;
+		}
+	}
+
+	/// <summary>
+	/// If a test bridge only applies for a very specific set of parameters,
+	/// they can be specified here. <br/>
+	/// Example: <code>[EnableForParams("Mode=Test", "SomeBoolVal", ...)]</code>
+	/// </summary>
+	[AttributeUsage(AttributeTargets.Class, AllowMultiple = true)]
+	public class EnableForParamsAttribute : Attribute
+	{
+		public List<KeyValuePair<string, string>> Params { get; private set; }
+
+		public EnableForParamsAttribute(params string[] InParams)
+		{
+			Params = InParams
+				.Select(P => P.Split('='))
+				.Select(S => new KeyValuePair<string, string>(S[0], S.Length > 1 ? S[1] : null))
+				.ToList();
+		}
+	}
+}
\ No newline at end of file
diff --git a/Engine/Plugins/Performance/AutomatedPerfTesting/Build/Scripts/AutomatedPerfTest.Util.cs b/Engine/Plugins/Performance/AutomatedPerfTesting/Build/Scripts/AutomatedPerfTest.Util.cs
new file mode 100644
--- /dev/null
+++ b/Engine/Plugins/Performance/AutomatedPerfTesting/Build/Scripts/AutomatedPerfTest.Util.cs
@@ -0,0 +1,110 @@
+// Copyright Epic Games, Inc. All Rights Reserved.
+
+using System;
+using System.Collections.Generic;
+using System.Linq;
+using System.Reflection;
+
+namespace AutomatedPerfTesting
+{
+	/// <summary>
+	/// Internal Utility Class
+	/// </summary>
+	internal static class Util
+	{
+		private static Assembly[] AssemblyCache = null;
+
+		private static Assembly[] GetAssemblies()
+		{
+			if (AssemblyCache == null)
+			{
+				AssemblyCache = AppDomain.CurrentDomain.GetAssemblies().Reverse().ToArray();
+			}
+
+			return AssemblyCache;
+		}
+
+		/// <summary>
+		/// Get array of types which match the given predicate.
+		/// </summary>
+		/// <param name="Predicate">Conditional predicate</param>
+		/// <returns>Array of types</returns>
+		public static Type[] GetTypes(Func<Type, bool> Predicate)
+		{
+			// Even though we cache the assemblies here, this function
+			// is still iterating through all types available within the
+			// assemblies to find matches. 
+			List<Type> FoundTypes = new List<Type>();
+			foreach (Assembly CandidateAssembly in GetAssemblies())
+			{
+				Type[] CandidateTypes = CandidateAssembly.GetTypes();
+				List<Type> Types = CandidateTypes
+						.Where(Predicate)
+						.ToList();
+
+				FoundTypes.AddRange(Types);
+			}
+
+			return FoundTypes.ToArray();
+		}
+
+		/// <summary>
+		/// Get reflected Type info of given type name if available.
+		/// </summary>
+		/// <typeparam name="TInterface">Interface the type is derived from</typeparam>
+		/// <param name="Name">Name of type</param>
+		/// <returns>Returns instance of Type if available, otherwise returns null</returns>
+		public static Type GetTypeWithInterface<TInterface>(string Name)
+		{
+			return GetTypeByName(Name, typeof(TInterface));
+		}
+
+		/// <summary>
+		/// Get reflected Type info of given type name if available. Searches assemblies
+		/// in current app domain. 
+		/// </summary>
+		/// <param name="Name">Name of type</param>
+		/// <param name="InterfaceType">Interface the type is derived from. Optional</param>
+		/// <returns>Returns instance of Type if available, otherwise returns null</returns>
+		public static Type GetTypeByName(string Name, Type InterfaceType = null)
+		{
+			if (InterfaceType == null)
+			{
+				InterfaceType = typeof(object);
+			}
+
+			foreach (Assembly CandidateAssembly in GetAssemblies())
+			{
+				Type CandidateType = CandidateAssembly.GetType(Name.Trim());
+				if (CandidateType != null && InterfaceType.IsAssignableFrom(CandidateType))
+				{
+					return CandidateType;
+				}
+			}
+
+			return null;
+		}
+
+		/// <summary>
+		/// Get array of types which is assignable from given interface
+		/// type. 
+		/// </summary>
+		/// <typeparam name="TInterface">Interface the types should implement</typeparam>
+		/// <returns></returns>
+		public static Type[] GetTypesWithInterface<TInterface>()
+		{
+			Type InterfaceType = typeof(TInterface);
+			return GetTypes(Type => Type.IsClass && InterfaceType.IsAssignableFrom(Type));
+		}
+
+		public static Type[] GetTypesWithAttribute<TAttribute>()
+		{
+			return GetTypesWithAttribute(typeof(TAttribute));
+		}
+
+		public static Type[] GetTypesWithAttribute(Type AttributeType)
+		{
+			return GetTypes(Type => Type.IsClass && Attribute.IsDefined(Type, AttributeType));
+		}
+	}
+}
diff --git a/Engine/Plugins/Performance/AutomatedPerfTesting/Build/Scripts/AutomatedPerfTestConfig.cs b/Engine/Plugins/Performance/AutomatedPerfTesting/Build/Scripts/AutomatedPerfTestConfig.cs
--- a/Engine/Plugins/Performance/AutomatedPerfTesting/Build/Scripts/AutomatedPerfTestConfig.cs
+++ b/Engine/Plugins/Performance/AutomatedPerfTesting/Build/Scripts/AutomatedPerfTestConfig.cs
@@ -12,7 +12,7 @@ using Log = EpicGames.Core.Log;
 
 namespace AutomatedPerfTest
 {
-	public class AutomatedPerfTestConfigBase : UnrealTestConfiguration
+	public class AutomatedPerfTestConfigBase : EpicGameTestConfig
 	{
 		/// <summary>
 		/// Used to override the test controller portion of the data source name
@@ -186,6 +186,24 @@ namespace AutomatedPerfTest
 		[AutoParamWithNames("", "AutomatedPerfTest.TestBuildVersion")]
 		public string TestBuildVersion = "";
 
+		/// <summary>
+		/// Overrides build version. If empty, default build version will be used. 
+		/// </summary>
+		[AutoParamWithNames("", "AutomatedPerfTest.BuildVersion")]
+		public string BuildVersion = "";
+
+		/// <summary>
+		/// Comma separated list of Bridges which will be initialized and invoked during configuration.
+		/// </summary>
+		[AutoParamWithNames("AutomatedPerfTest.NullBridge", "AutomatedPerfTest.ConfigBridges")]
+		public string ConfigBridges = "";
+
+		/// <summary>
+		/// Overrides managed account pool instance with the one provided. 
+		/// </summary>
+		[AutoParamWithNames("", "AutomatedPerfTest.ManagedAccountPool")]
+		public string ManagedAcountPool;
+
 		public string DataSourceName;
 
 		/// <summary>
diff --git a/Engine/Plugins/Performance/AutomatedPerfTesting/Build/Scripts/AutomatedPerfTestNode.cs b/Engine/Plugins/Performance/AutomatedPerfTesting/Build/Scripts/AutomatedPerfTestNode.cs
--- a/Engine/Plugins/Performance/AutomatedPerfTesting/Build/Scripts/AutomatedPerfTestNode.cs
+++ b/Engine/Plugins/Performance/AutomatedPerfTesting/Build/Scripts/AutomatedPerfTestNode.cs
@@ -17,6 +17,7 @@ using System.Reflection;
 using System.Text;
 using static Gauntlet.HordeReport.AutomatedTestSessionData;
 using AutomatedTestSessionData = Gauntlet.HordeReport.AutomatedTestSessionData;
+using AutomatedPerfTesting;
 
 namespace AutomatedPerfTest
 {
@@ -44,6 +45,7 @@ namespace AutomatedPerfTest
 			Log.Info("Your Test GUID is :\n" + TestGuid.ToString() + '\n');
 
 			InitHandledErrors();
+			AutoTestBridge.ActivateTestBridges(GetType(), InContext.BuildInfo.ProjectName);
 
 			LogParser = null;
 		}
@@ -101,14 +103,14 @@ namespace AutomatedPerfTest
 		/// </summary>
 		/// <param name="Platform"></param>
 		/// <returns>{ProjectName}/Saved/Performance/{SubTest}/{Platform}</returns>
-		public string GetPerformanceReportArtifactOutputPath(UnrealTargetPlatform Platform)
-		{ 
-			if(BaseOutputPath == null)
+		public string GetPerformanceReportArtifactOutputPath(UnrealTargetPlatform Platform, TConfigClass Config = null)
+		{
+			if (BaseOutputPath == null)
 			{
 				return TempPerfCSVDir.FullName;
 			}
 
-			return Path.Combine(BaseOutputPath, GetSubtestName(GetCachedConfiguration()), Platform.ToString());
+			return Path.Combine(BaseOutputPath, GetSubtestName(Config ?? GetCachedConfiguration()), Platform.ToString());
 		}
 
 		/// <summary>
@@ -906,6 +908,41 @@ namespace AutomatedPerfTest
 				ClientRole.CommandLineParams.Add("WaitForDebugger");
 			}
 
+			// Get config bridge if specified. This is useful if we want project
+			// specific configuration which can be injected via params 
+			if (!string.IsNullOrEmpty(Config.ConfigBridges))
+			{
+				AutoTestBridge.Initialize(Config.ConfigBridges.Split(","));
+			}
+
+			AutoTestBridge.Configure(Context, Config);
+			AutoTestBridge.ConfigureClient(ClientRole, Context, Config);
+
+			if (AutoTestBridge.IsServerRoleRequired())
+			{
+				string DefaultClientSenderId = $"{Context.BuildInfo.ProjectName}Client";
+				ConfigureServerRole(Config);
+
+				// Some fun necessary commandlines to make sure our sockets are working properly. 
+				string BindAddress = Context.TestParams.ParseValue("bindaddress", "any");
+				ClientRole.CommandLineParams.Add("ini:Engine:[HTTPServer.Listeners]:DefaultBindAddress", BindAddress);
+				ClientRole.CommandLineParams.Add("ini:Engine:[HTTPServer.Listeners]:DefaultReuseAddressAndPort", true);
+				ClientRole.CommandLineParams.Add("logcmds", "LogSockets VeryVerbose");
+
+				// Tell client and server where to send messages to communicate with Gauntlet, and give them standardized names.
+				ClientRole.CommandLineParams.AddUnique("externalrpclistenaddress", RpcExecutor.GetListenAddressAndPort());
+				ClientRole.CommandLineParams.AddUnique("rpcsenderid", DefaultClientSenderId);
+			}
+
+			if (!string.IsNullOrEmpty(Config.ManagedAcountPool))
+			{
+				OverrideManagedAccountPool(Config.ManagedAcountPool);
+			}
+
+			UnrealTargetPlatform Platform = Context.GetRoleContext(UnrealTargetRole.Client).Platform;
+			string ArtifactOutputPath = GetPerformanceReportArtifactOutputPath(Platform, Config);
+			ClientRole.CommandLineParams.AddUnique("AutomatedPerfTest.ArtifactOutputPath", ArtifactOutputPath);
+
 			return Config;
 		}
 
@@ -1064,15 +1101,18 @@ namespace AutomatedPerfTest
 			{
 				Config.TestBuildVersion = OriginalBuildName;
 			}
-			else
+
+			if(string.IsNullOrEmpty(Config.BuildVersion))
 			{
 				// If build version has been overridden for any reason,
 				// this ensures that the CSV metadata matches the test
 				// build version since this var is used while instantiating
 				// perf report importer. 
-				OriginalBuildName = Config.TestBuildVersion;
+				Config.BuildVersion = OriginalBuildName;
 			}
 
+			OriginalBuildName = Config.BuildVersion;
+
 			Log.Info($"Using Test Build Version {Config.TestBuildVersion}");
 
 			// Add CSV metadata
@@ -1081,7 +1121,7 @@ namespace AutomatedPerfTest
 				$"testname={Context.BuildInfo.ProjectName}",
 				"gauntletTestType=AutomatedPerfTest",
 				$"gauntletSubTest={GetSubtestName(Config)}",
-				"testBuildIsPreflight=" + (ReportGenUtils.IsTestingPreflightBuild(OriginalBuildName) ? "1" : "0"),
+				"testBuildIsPreflight=" + (ReportGenUtils.IsTestingPreflightBuild(Config.TestBuildVersion) ? "1" : "0"),
 				$"testBuildVersion={Config.TestBuildVersion}",
 				"testconfigname=" + Config.TestConfigName,
 			];
@@ -1100,6 +1140,36 @@ namespace AutomatedPerfTest
 			ClientRole.CommandLineParams.Add("csvMetadata", "\"" + String.Join(",", CsvMetadata) + "\"");
 		}
 
+		private void ConfigureServerRole(TConfigClass Config)
+		{
+			// Assumed default base port. TODO: Should be overridable. 
+			const int DefaultBasePort = 12321;
+			string DefaultServerSenderId = $"{Context.BuildInfo.ProjectName}Server";
+			UnrealTestRole ServerRole = Config.RequireRole(UnrealTargetRole.Server);
+
+			ServerRole.CommandLineParams.AddUnique("rpcport", DefaultBasePort + 1);
+			ServerRole.CommandLineParams.AddUnique("externalrpclistenaddress", RpcExecutor.GetListenAddressAndPort());
+			ServerRole.CommandLineParams.AddUnique("rpcsenderid", DefaultServerSenderId);
+
+			AutoTestBridge.ConfigureServer(ServerRole, Context, Config);
+		}
+
+		private static void OverrideManagedAccountPool(string ManagedAccountPool)
+		{
+			if (!string.IsNullOrEmpty(ManagedAccountPool))
+			{
+				Type ManagedAccountPoolType = Util.GetTypeWithInterface<IManagedAccountPool>(ManagedAccountPool);
+				MethodInfo InitGenericMethod = typeof(AccountPool).GetMethods(BindingFlags.Static | BindingFlags.Public)
+					.Where(Info => Info.Name == "Initialize" && Info.IsGenericMethod)
+					.First();
+
+				// Invokes Initialize<T>() with given Managed Account Pool implementation
+				// type which will override the existing instance, if any.
+				MethodInfo InitMethodRef = InitGenericMethod.MakeGenericMethod(ManagedAccountPoolType);
+				InitMethodRef.Invoke(null, null);
+			}
+		}
+
 		static protected void GetConfigValues(UnrealTestContext Context, string IniSection, string IniElement, out IReadOnlyList<string> Values)
 		{
 			IniConfigUtil.GetConfigHierarchy(Context, ConfigHierarchyType.Engine).TryGetValues(IniSection, IniElement, out Values);
diff --git a/Engine/Plugins/Performance/AutomatedPerfTesting/Source/AutomatedPerfTesting/Private/AutomatedPerfTestControllerBase.cpp b/Engine/Plugins/Performance/AutomatedPerfTesting/Source/AutomatedPerfTesting/Private/AutomatedPerfTestControllerBase.cpp
--- a/Engine/Plugins/Performance/AutomatedPerfTesting/Source/AutomatedPerfTesting/Private/AutomatedPerfTestControllerBase.cpp
+++ b/Engine/Plugins/Performance/AutomatedPerfTesting/Source/AutomatedPerfTesting/Private/AutomatedPerfTestControllerBase.cpp
@@ -21,6 +21,7 @@
 #include "ProfilingDebugging/TraceScreenshot.h"
 #include "UnrealClient.h"
 #include "UnrealEngine.h"
+#include "GauntletModule.h"
 
 #include UE_INLINE_GENERATED_CPP_BY_NAME(AutomatedPerfTestControllerBase)
 
@@ -119,6 +120,7 @@ UAutomatedPerfTestControllerBase::UAutomatedPerfTestControllerBase(const FObject
 	, bRequestsVideoCapture(false)
 	, bRequestsLockedDynRes(false)
 	, InsightsRegionID(0)
+	, ArtifactOutputPath()
     , CSVOutputMode(EAutomatedPerfTestCSVOutputMode::Single)
 {
 	// cache this off once, so that it's consistent throughout a session
@@ -203,7 +205,8 @@ bool UAutomatedPerfTestControllerBase::TryStopInsightsTrace()
 
 bool UAutomatedPerfTestControllerBase::TryStartCSVProfiler()
 {
-	return TryStartCSVProfiler(GetCSVFilename());
+	const FString DestinationFolder = ArtifactOutputPath.IsEmpty() ? FString() : FPaths::Combine(ArtifactOutputPath, "CSV");
+	return TryStartCSVProfiler(GetCSVFilename(), DestinationFolder);
 }
 
 bool UAutomatedPerfTestControllerBase::TryStartCSVProfiler(FString CSVFileName, const FString& DestinationFolder, int32 Frames)
@@ -222,6 +225,9 @@ bool UAutomatedPerfTestControllerBase::TryStartCSVProfiler(FString CSVFileName,
 		CsvProfiler->SetMetadata(TEXT("Datetime"), *TestDatetime);
 		CsvProfiler->SetMetadata(TEXT("ResX"), *FString::FromInt(GSystemResolution.ResX));
 		CsvProfiler->SetMetadata(TEXT("ResY"), *FString::FromInt(GSystemResolution.ResY));
+
+		UE_LOG(LogGauntlet, Log, TEXT("CSV Profiler Destination Filename: %s, Destination Folder: %s"), 
+				*CSVFileName, DestinationFolder.IsEmpty()? TEXT("<Default>") : *DestinationFolder);
 		
 		CsvProfiler->BeginCapture(Frames, DestinationFolder, CSVFileName);
 		CsvProfiler->SetDeviceProfileName(GetDeviceProfile());
@@ -348,6 +354,7 @@ void UAutomatedPerfTestControllerBase::SetupTest()
 	// Subclasses should implement their own transitions from SetupTest to RunTest depending on their needs
 
 	UE_LOG(LogAutomatedPerfTest, Log, TEXT("CSV Output Mode: %s"), *UEnum::GetValueAsString(CSVOutputMode))
+	UE_LOG(LogGauntlet, Log, TEXT("Setup Test: %s"), *GetTestID());
 }
 
 void UAutomatedPerfTestControllerBase::RunTest()
@@ -358,6 +365,8 @@ void UAutomatedPerfTestControllerBase::RunTest()
 	{
 		IAutomatedPerfTestInterface::Execute_RunTest(GameMode);
 	}
+
+	UE_LOG(LogGauntlet, Log, TEXT("Running Test: %s"), *GetTestID());
 }
 
 void UAutomatedPerfTestControllerBase::TeardownTest(bool bExitAfterTeardown)
@@ -373,13 +382,27 @@ void UAutomatedPerfTestControllerBase::TeardownTest(bool bExitAfterTeardown)
 	{
 		TriggerExitAfterDelay();
 	}
+
+	UE_LOG(LogGauntlet, Log, TEXT("Teardown Test: %s"), *GetTestID());
 }
 
 void UAutomatedPerfTestControllerBase::TriggerExitAfterDelay()
 {
-	FTimerHandle UnusedHandle;
-	
-	GetWorld()->GetTimerManager().SetTimer(UnusedHandle, this, &UAutomatedPerfTestControllerBase::Exit, 1.0, false, GetDefault<UAutomatedPerfTestProjectSettings>()->TeardownToExitDelay);	
+	const float Delay = GetDefault<UAutomatedPerfTestProjectSettings>()->TeardownToExitDelay;
+
+	// Gauntlet test controllers are UObject types.
+	TObjectPtr<UAutomatedPerfTestControllerBase> ControllerInstance(this);
+	FTSTicker::GetCoreTicker().AddTicker(TEXT("ReplayComplete"), Delay, [ControllerInstance](float)
+	{
+		// Capturing this here is safe as this controller is only destroyed after 
+		// Exit is called below and the Gauntlet test ends. When running Gauntlet 
+		// tests, the controller instances are owned by the Gauntlet Module 
+		if(ControllerInstance)
+		{
+			ControllerInstance->Exit();
+		}
+		return false;
+	});
 }
 
 void UAutomatedPerfTestControllerBase::Exit()
@@ -610,10 +633,16 @@ void UAutomatedPerfTestControllerBase::OnInit()
 		UE_LOG(LogAutomatedPerfTest, Log, TEXT("Locking dynamic res requested"))
 		bRequestsLockedDynRes = true;
 	}
+	if (FParse::Value(FCommandLine::Get(), TEXT("AutomatedPerfTest.ArtifactOutputPath="), ArtifactOutputPath))
+	{
+		UE_LOG(LogAutomatedPerfTest, Log, TEXT("Artifact Output Path: %s"), *ArtifactOutputPath)
+	}
 	
 	FWorldDelegates::OnPreWorldInitialization.AddUObject(this, &ThisClass::OnPreWorldInitializeInternal);
 
 	InitializeInsights();
+
+	UE_LOG(LogGauntlet, Log, TEXT("Initializing Test: %s"), *GetTestID());
 }
 
 void UAutomatedPerfTestControllerBase::OnTick(float TimeDelta)
diff --git a/Engine/Plugins/Performance/AutomatedPerfTesting/Source/AutomatedPerfTesting/Private/ProfileGo/ProfileGo.cpp b/Engine/Plugins/Performance/AutomatedPerfTesting/Source/AutomatedPerfTesting/Private/ProfileGo/ProfileGo.cpp
--- a/Engine/Plugins/Performance/AutomatedPerfTesting/Source/AutomatedPerfTesting/Private/ProfileGo/ProfileGo.cpp
+++ b/Engine/Plugins/Performance/AutomatedPerfTesting/Source/AutomatedPerfTesting/Private/ProfileGo/ProfileGo.cpp
@@ -260,7 +260,7 @@ void UProfileGo::AddNewProfileGoCollection(FProfileGoCollectionAPT NewCollection
 	}
 
 	TArray<FString> CollectionScenarios;
-	NewCollection.Scenarios.ParseIntoArray(CollectionScenarios, TEXT(","), true);
+	CollectionScenarios = NewCollection.Scenarios;
 	for (FString ScenarioName : CollectionScenarios)
 	{
 		if (!GetScenarioTemplates().ContainsByPredicate([ScenarioName](FProfileGoScenarioAPT Scenario) { return Scenario.Name.Equals(ScenarioName); }))
@@ -348,7 +348,7 @@ bool UProfileGo::TryFindScenarios(const TCHAR* ProfileName, TArray<FProfileGoSce
 
 	if (CollectionExists)
 	{
-		CollectionExists->Scenarios.ParseIntoArray(ScenarioNames, TEXT(","), true);
+		ScenarioNames = CollectionExists->Scenarios;
 	}
 	else
 	{
@@ -1520,11 +1520,7 @@ void UProfileGo::RegisterInternalGeneratedScenarioDelegates()
 								NewProfileScenario.AutoGenerated = true;
 
 								// Add to the collection
-								if (NewCollection.Scenarios.Len() > 0)
-								{
-									NewCollection.Scenarios.AppendChar(',');
-								}
-								NewCollection.Scenarios.Append(ScenarioName);
+								NewCollection.Scenarios.Add(ScenarioName);
 
 								// Add all grid point scenarios regardless of if we need them in this collection.
 								AddNewProfileGoScenario(NewProfileScenario);
@@ -1713,7 +1709,7 @@ void UProfileGo::RegisterInternalGeneratedScenarioDelegates()
 
 			FProfileGoCollectionAPT NewCollection;
 			NewCollection.Name = ProfileName;
-			NewCollection.Scenarios = FString::Join(Scenarios, TEXT(","));
+			NewCollection.Scenarios = Scenarios;
 
 			AddNewProfileGoCollection(NewCollection);
 		}
@@ -2090,7 +2086,7 @@ void UProfileGo::AddScenariosInLevel(UWorld* World)
 		{
 			FVector Translation = ActorAPT->GetActorLocation();
 			FRotator Rotator = ActorAPT->GetActorRotation();
-			FString CollectionName = ActorAPT->Group;
+			FString CollectionName = ActorAPT->CollectionName.IsEmpty() ? TEXT("Default") : ActorAPT->CollectionName;
 
 			FString ScenarioName = ActorAPT->GetActorLabel();
 
@@ -2116,10 +2112,6 @@ void UProfileGo::AddScenariosInLevel(UWorld* World)
 				FProfileGoScenarioAPT* Scenario = GetScenarioTemplates().FindByPredicate([ScenarioName](FProfileGoScenarioAPT Scenario) { return Scenario.Name.Equals(ScenarioName); });
 				Scenario->Orientation = Rotator;
 				Scenario->Position = Translation;
-
-				// Assuming that if scenario exists, it's already part of a collection
-				// TODO: May need to revisit this. 
-				continue;
 			}
 
 			if (!CollectionExists(CollectionName))
@@ -2137,13 +2129,15 @@ void UProfileGo::AddScenariosInLevel(UWorld* World)
 
 			if (Collection != nullptr)
 			{
-				if (!Collection->Scenarios.Contains(ScenarioName))
+				// Add scenario to corresponding collection if it does not exist.
+				const bool bContainsScenario = Collection->Scenarios.ContainsByPredicate([ScenarioName](FString Scenario)
+				{ 
+					return Scenario.Equals(ScenarioName, ESearchCase::IgnoreCase); 
+				});
+
+				if (!bContainsScenario)
 				{
-					if (Collection->Scenarios.Len() > 0)
-					{
-						Collection->Scenarios.AppendChar(',');
-					}
-					Collection->Scenarios.Append(ScenarioName);
+					Collection->Scenarios.Add(ScenarioName);
 				}
 			}
 		}
diff --git a/Engine/Plugins/Performance/AutomatedPerfTesting/Source/AutomatedPerfTesting/Public/AutomatedPerfTestControllerBase.h b/Engine/Plugins/Performance/AutomatedPerfTesting/Source/AutomatedPerfTesting/Public/AutomatedPerfTestControllerBase.h
--- a/Engine/Plugins/Performance/AutomatedPerfTesting/Source/AutomatedPerfTesting/Public/AutomatedPerfTestControllerBase.h
+++ b/Engine/Plugins/Performance/AutomatedPerfTesting/Source/AutomatedPerfTesting/Public/AutomatedPerfTestControllerBase.h
@@ -163,6 +163,7 @@ private:
 	bool bRequestsVideoCapture;
 	bool bRequestsLockedDynRes;
 	uint64 InsightsRegionID;
+	FString ArtifactOutputPath;
 
 	FText VideoRecordingTitle;
 	
diff --git a/Engine/Plugins/Performance/AutomatedPerfTesting/Source/AutomatedPerfTesting/Public/ProfileGo/ProfileGo.h b/Engine/Plugins/Performance/AutomatedPerfTesting/Source/AutomatedPerfTesting/Public/ProfileGo/ProfileGo.h
--- a/Engine/Plugins/Performance/AutomatedPerfTesting/Source/AutomatedPerfTesting/Public/ProfileGo/ProfileGo.h
+++ b/Engine/Plugins/Performance/AutomatedPerfTesting/Source/AutomatedPerfTesting/Public/ProfileGo/ProfileGo.h
@@ -110,7 +110,7 @@ struct FProfileGoCollectionAPT
 	FString Name;
 
 	UPROPERTY(Config, EditAnywhere, BlueprintReadWrite, Category = "APT")
-	FString Scenarios;
+	TArray<FString> Scenarios;
 };
 
 USTRUCT(BlueprintType, Category = "APT")
diff --git a/Engine/Plugins/Performance/AutomatedPerfTesting/Source/AutomatedPerfTesting/Public/StaticCameraTests/AutomatedPerfTestStaticCamera.h b/Engine/Plugins/Performance/AutomatedPerfTesting/Source/AutomatedPerfTesting/Public/StaticCameraTests/AutomatedPerfTestStaticCamera.h
--- a/Engine/Plugins/Performance/AutomatedPerfTesting/Source/AutomatedPerfTesting/Public/StaticCameraTests/AutomatedPerfTestStaticCamera.h
+++ b/Engine/Plugins/Performance/AutomatedPerfTesting/Source/AutomatedPerfTesting/Public/StaticCameraTests/AutomatedPerfTestStaticCamera.h
@@ -13,5 +13,5 @@ public:
 	AAutomatedPerfTestStaticCamera(const FObjectInitializer& ObjectInitializer);
 
 	UPROPERTY(EditAnywhere, BlueprintReadWrite, Category = "ProfileGo | Scenario", meta = (DefaultValue = "Default"))
-	FString Group;
+	FString CollectionName;
 };
