APT: Fix CSV path and need for Managed Account in non-server based tests

* Assumption of server role caused managed account to be needed for non-server based automated testing. This fixes it by making -NoMCP default.
* Fix CSV path not found - this was because on build machines, we do not need to override the artifact output path, this change accounts for that.
* Fix for build version mismatch with test build version.

#rnx

#tests local, preflight

diff --git a/Engine/Plugins/Performance/AutomatedPerfTesting/Build/Inc/AutomatedPerfTestCommonSettings.xml b/Engine/Plugins/Performance/AutomatedPerfTesting/Build/Inc/AutomatedPerfTestCommonSettings.xml
--- a/Engine/Plugins/Performance/AutomatedPerfTesting/Build/Inc/AutomatedPerfTestCommonSettings.xml
+++ b/Engine/Plugins/Performance/AutomatedPerfTesting/Build/Inc/AutomatedPerfTestCommonSettings.xml
@@ -5,7 +5,8 @@
 	<EnvVar Name="UE_HORDE_JOBID"/>
 	<EnvVar Name="UE_HORDE_STEPID"/>
 	<Property Name="JobDetails" Value="$(HordeDashboardURL)/job/$(UE_HORDE_JOBID)?step=$(UE_HORDE_STEPID)"/>
-	<Property Name="APTBuildVersion" Value="$(EscapedBranch)-CL-$(Change)-$(ProjectName)$(PreflightSuffix)"/> 
+	<Property Name="APTBuildVersion" Value="$(EscapedBranch)-CL-$(Change)-$(ProjectName)"/> 
+	<Property Name="APTTestBuildVersion" Value="$(EscapedBranch)-CL-$(Change)-$(ProjectName)$(PreflightSuffix)"/> 
 	<!-- Insights Setup for Automated Perf Testing -->
 	<!-- since we need to tell UAT whether or not to using the shipping configuration of Insights when processing -->
 	<!-- insights traces after the conclusion of the test --> 
@@ -37,12 +38,12 @@
 	<Option Name="APTIgnoreTestBuildLogging" Restrict="true|false" DefaultValue="true" Description="If true, tests run in Test build configuration will not fail if logs are empty." />
 	<Option Name="PerfOutputPath" DefaultValue="" Description="Path where Perf Test artifacts (Logs, CSVs, Reports, Cache) will be stored. Override to a common network-accessible path when running on build machines if needed"/>
 	<Option Name="ProjectPath" DefaultValue="$(DefaultProjectPath)" Restrict=".+" Description="Path to the folder that contains your project"/>
-	<Option Name="TestBuildVersion" DefaultValue="$(APTBuildVersion)" Description="Overrides test build version"/>
+	<Option Name="TestBuildVersion" DefaultValue="$(APTTestBuildVersion)" Description="Overrides test build version"/>
 	<Option Name="ExtraAutomatedPerformanceCommonArgs" DefaultValue="" Description="Additional Common arguments passed to Gauntlet command line"/>
 	<Option Name="NullOSSPlatforms" DefaultValue="Android" Description=""/>
 
 	<!-- Override build version with given value. If this value is empty, APT will use default build version-->
-	<Property Name="BuildVersion" Value="$(TestBuildVersion)"/> 
+	<Property Name="BuildVersion" Value="$(APTBuildVersion)"/> 
 
 	<!-- Plugin-provided tests -->
 	<Option Name="WithAPTSequenceTests" DefaultValue="true" Restrict="true|false" Description="Whether to run Auto Perf Test sequence flythrough test." />
@@ -74,7 +75,8 @@
 	<Property Name="AutomatedPerformanceCommonArgs" Value="$(AutomatedPerformanceCommonArgs) -AutomatedPerfTest.PerfOutputPath=&quot;$(PerfOutputPath)&quot;" If="'$(PerfOutputPath)' != ''"/>
 	<Property Name="AutomatedPerformanceCommonArgs" Value="$(AutomatedPerformanceCommonArgs) $(ExtraAutomatedPerformanceCommonArgs)" If="'$(ExtraAutomatedPerformanceCommonArgs)' != ''"/>
 	<Property Name="AutomatedPerformanceCommonArgs" Value="$(AutomatedPerformanceCommonArgs) -Branch=&quot;$(EscapedBranch)&quot; -CL=$(Change)"/>
-	<Property Name="AutomatedPerformanceCommonArgs" Value="$(AutomatedPerformanceCommonArgs) -AutomatedPerfTest.TestBuildVersion=&quot;$(BuildVersion)&quot;"/>
+	<Property Name="AutomatedPerformanceCommonArgs" Value="$(AutomatedPerformanceCommonArgs) -AutomatedPerfTest.TestBuildVersion=&quot;$(TestBuildVersion)&quot;"/>
+	<Property Name="AutomatedPerformanceCommonArgs" Value="$(AutomatedPerformanceCommonArgs) -AutomatedPerfTest.BuildVersion=&quot;$(BuildVersion)&quot;"/>
 	<Property Name="AutomatedPerformanceCommonArgs" Value="$(AutomatedPerformanceCommonArgs) -AutomatedPerfTest.NullOSSPlatforms=&quot;$(NullOSSPlatforms)&quot;" If="'$(NullOSSPlatforms)' != ''"/>
 
 	<Property Name="PerfTestCommonArgs"     Value="$(AutomatedPerformanceCommonArgs) -iterations=$(IterationsPerf)"/>
diff --git a/Engine/Plugins/Performance/AutomatedPerfTesting/Build/Scripts/AutomatedPerfTestNode.cs b/Engine/Plugins/Performance/AutomatedPerfTesting/Build/Scripts/AutomatedPerfTestNode.cs
--- a/Engine/Plugins/Performance/AutomatedPerfTesting/Build/Scripts/AutomatedPerfTestNode.cs
+++ b/Engine/Plugins/Performance/AutomatedPerfTesting/Build/Scripts/AutomatedPerfTestNode.cs
@@ -810,7 +810,6 @@ namespace AutomatedPerfTest
 		{
 			TConfigClass Config = base.GetConfiguration();
 			Config.MaxDuration = Context.TestParams.ParseValue("MaxDuration", 60 * 60);  // 1 hour max
-			//Config.NoMCP = true;  
 
 			if (string.IsNullOrEmpty(Config.PerfOutputPath))
 			{
@@ -908,6 +907,10 @@ namespace AutomatedPerfTest
 				ClientRole.CommandLineParams.Add("WaitForDebugger");
 			}
 
+			// Setting default to true as tests which do not require
+			// server role do not need MCP.
+			Config.NoMCP = true;
+
 			// Get config bridge if specified. This is useful if we want project
 			// specific configuration which can be injected via params 
 			if (!string.IsNullOrEmpty(Config.ConfigBridges))
@@ -920,6 +923,8 @@ namespace AutomatedPerfTest
 
 			if (AutoTestBridge.IsServerRoleRequired())
 			{
+				// Override default if server role is required. 
+				Config.NoMCP = Context.TestParams.ParseParam("NoMCP");
 				string DefaultClientSenderId = $"{Context.BuildInfo.ProjectName}Client";
 				ConfigureServerRole(Config);
 
@@ -939,9 +944,17 @@ namespace AutomatedPerfTest
 				OverrideManagedAccountPool(Config.ManagedAcountPool);
 			}
 
-			UnrealTargetPlatform Platform = Context.GetRoleContext(UnrealTargetRole.Client).Platform;
-			string ArtifactOutputPath = GetPerformanceReportArtifactOutputPath(Platform, Config);
-			ClientRole.CommandLineParams.AddUnique("AutomatedPerfTest.ArtifactOutputPath", ArtifactOutputPath);
+			if(!IsBuildMachine)
+			{
+				// If we are on a build machine, artifacts are stored and handled as expected. On build machines, reports
+				// are eventually copied to this location before being processed. When generating local reports, it is
+				// possible artifacts are being output to a place not expected by APT, so we let the runtime know to 
+				// output the artifacts to this location directly. 
+
+				UnrealTargetPlatform Platform = Context.GetRoleContext(UnrealTargetRole.Client).Platform;
+				string ArtifactOutputPath = GetPerformanceReportArtifactOutputPath(Platform, Config);
+				ClientRole.CommandLineParams.AddUnique("AutomatedPerfTest.ArtifactOutputPath", ArtifactOutputPath);
+			}
 
 			return Config;
 		}
