Add test performance relevant metadata

#rnx

diff --git a/Engine/Plugins/Performance/AutomatedPerfTesting/Build/Scripts/AutomatedPerfTestConfig.cs b/Engine/Plugins/Performance/AutomatedPerfTesting/Build/Scripts/AutomatedPerfTestConfig.cs
--- a/Engine/Plugins/Performance/AutomatedPerfTesting/Build/Scripts/AutomatedPerfTestConfig.cs
+++ b/Engine/Plugins/Performance/AutomatedPerfTesting/Build/Scripts/AutomatedPerfTestConfig.cs
@@ -235,6 +235,10 @@ namespace AutomatedPerfTest
 
 		public override void ApplyToConfig(UnrealAppConfig AppConfig, UnrealSessionRole ConfigRole, IEnumerable<UnrealSessionRole> OtherRoles)
 		{
+			// disable MCP login
+			NoMCP = true;
+			RequiresLogin = false;
+
 			base.ApplyToConfig(AppConfig, ConfigRole, OtherRoles);
 
 			// Let the app know if we need to use extra debug memory before we 
diff --git a/Engine/Plugins/Performance/AutomatedPerfTesting/Build/Scripts/AutomatedPerfTestNode.cs b/Engine/Plugins/Performance/AutomatedPerfTesting/Build/Scripts/AutomatedPerfTestNode.cs
--- a/Engine/Plugins/Performance/AutomatedPerfTesting/Build/Scripts/AutomatedPerfTestNode.cs
+++ b/Engine/Plugins/Performance/AutomatedPerfTesting/Build/Scripts/AutomatedPerfTestNode.cs
@@ -493,9 +493,15 @@ namespace AutomatedPerfTest
 				AddLinkToReport(Report, PRSReportName, PRSReportLink);
 			}
 
+			Report.SetMetadata("PerfType", GetSubtestName(Config));
+			Report.SetMetadata("PerfTest", string.IsNullOrEmpty(Config.TestID)? "Default" : Config.TestID);
+
 			return Report;
 		}
 
+		/// <inheritdoc/>
+		protected override string HordeReportTestKey => base.HordeReportTestKey + (string.IsNullOrEmpty(GetCachedConfiguration().TestID) ? "" : $".{GetCachedConfiguration().TestID}");
+
 		/// <summary>
 		/// Get report type for current configuration
 		/// </summary>
