Disabled r.Substrate.AsyncClassification for -ReducedAsyncCompute.
Amendment to 45723288 that only disabled it for CSVProfiling.

diff --git a/Engine/Plugins/Performance/AutomatedPerfTesting/Build/Scripts/AutomatedPerfTestNode.cs b/Engine/Plugins/Performance/AutomatedPerfTesting/Build/Scripts/AutomatedPerfTestNode.cs
--- a/Engine/Plugins/Performance/AutomatedPerfTesting/Build/Scripts/AutomatedPerfTestNode.cs
+++ b/Engine/Plugins/Performance/AutomatedPerfTesting/Build/Scripts/AutomatedPerfTestNode.cs
@@ -886,6 +886,7 @@ namespace AutomatedPerfTest
 				ClientRole.CommandLineParams.AddOrAppendParamValue("execcmds", "r.AmbientOcclusion.Compute 1"); // 1 here means compute on the graphics pipe
 				ClientRole.CommandLineParams.AddOrAppendParamValue("execcmds", "r.LocalFogVolume.TileCullingUseAsync 0");
 				ClientRole.CommandLineParams.AddOrAppendParamValue("execcmds", "r.SkyAtmosphereASyncCompute 0");
+				ClientRole.CommandLineParams.AddOrAppendParamValue("execcmds", "r.Substrate.AsyncClassification 0");
 			}
 
 			if (Config.DoGPUPerf || Config.LockDynamicRes)
