APT: Fix PL2 not generating reports for local build

Also fixes the issue of test type not persisting in a Launch Profile

#tests Running APT tests on Win64 and consoles

diff --git a/Engine/Plugins/Developer/ProjectLauncher/Source/CommonLaunchExtensions/Private/AutomatedPerfTest/AutomatedPerfTestLaunchExtension.cpp b/Engine/Plugins/Developer/ProjectLauncher/Source/CommonLaunchExtensions/Private/AutomatedPerfTest/AutomatedPerfTestLaunchExtension.cpp
--- a/Engine/Plugins/Developer/ProjectLauncher/Source/CommonLaunchExtensions/Private/AutomatedPerfTest/AutomatedPerfTestLaunchExtension.cpp
+++ b/Engine/Plugins/Developer/ProjectLauncher/Source/CommonLaunchExtensions/Private/AutomatedPerfTest/AutomatedPerfTestLaunchExtension.cpp
@@ -193,7 +193,7 @@ void FAutomatedPerfTestLaunchExtensionInstance::CustomizeTree(ProjectLauncher::F
 	// for the test types above as well to ensure test types can be added on and is 
 	// extensible in the future.
 
-	constexpr EConfig ConfigTypeCommon = EConfig::User_PerProfile;
+	constexpr EConfig ConfigTypeCommon = EConfig::PerProfile;
 	FLaunchProfileTreeNode& TreeNode = 
 		AddDefaultHeading(ProfileTreeData)
 		.AddBoolean(LOCTEXT("APTEnableCSVLabel", "Enable CSV Profiler"), 
@@ -523,6 +523,7 @@ void FAutomatedPerfTestLaunchExtensionInstance::SetTestType(EAutomatedPerfTestTy
 {
 	CurrentTestType = Type;
 	SetConfigInteger(EConfig::PerProfile, TEXT("TestType"), (int32)Type);
+	BroadcastPropertyChanged();
 }
 
 EAutomatedPerfTestType FAutomatedPerfTestLaunchExtensionInstance::GetTestType() const
diff --git a/Engine/Plugins/Developer/ProjectLauncher/Source/CommonLaunchExtensions/Private/AutomatedPerfTest/AutomatedPerfTestLaunchExtension.h b/Engine/Plugins/Developer/ProjectLauncher/Source/CommonLaunchExtensions/Private/AutomatedPerfTest/AutomatedPerfTestLaunchExtension.h
--- a/Engine/Plugins/Developer/ProjectLauncher/Source/CommonLaunchExtensions/Private/AutomatedPerfTest/AutomatedPerfTestLaunchExtension.h
+++ b/Engine/Plugins/Developer/ProjectLauncher/Source/CommonLaunchExtensions/Private/AutomatedPerfTest/AutomatedPerfTestLaunchExtension.h
@@ -52,7 +52,7 @@ private:
 		const auto IsTestType = [this]()
 		{
 			const bool bEnabled = IsTestActive() || GetConfigBool(EConfig::PerProfile, TEXT("Enabled"));
-			return CurrentTestType == TestType && bEnabled;
+			return GetTestType() == TestType && bEnabled;
 		};
 
 		ProjectLauncher::FLaunchProfileTreeNode::FCallbacks Callbacks
diff --git a/Engine/Plugins/Performance/AutomatedPerfTesting/Build/Scripts/AutomatedPerfTestConfig.cs b/Engine/Plugins/Performance/AutomatedPerfTesting/Build/Scripts/AutomatedPerfTestConfig.cs
--- a/Engine/Plugins/Performance/AutomatedPerfTesting/Build/Scripts/AutomatedPerfTestConfig.cs
+++ b/Engine/Plugins/Performance/AutomatedPerfTesting/Build/Scripts/AutomatedPerfTestConfig.cs
@@ -204,6 +204,14 @@ namespace AutomatedPerfTest
 		[AutoParamWithNames("", "AutomatedPerfTest.ManagedAccountPool")]
 		public string ManagedAcountPool;
 
+		/// <summary>
+		/// Overrides default artifact output path by passing the artifact path to the client instance. 
+		/// This is useful if the project has set the artifact to output to a non-standard default output 
+		/// path causing report generation to fail because of missing CSV artifacts.
+		/// </summary>
+		[AutoParamWithNames(false, "AutomatedPerfTest.OverrideArtifactPath")]
+		public bool OverrideArtifactOutputPath;
+
 		public string DataSourceName;
 
 		/// <summary>
diff --git a/Engine/Plugins/Performance/AutomatedPerfTesting/Build/Scripts/AutomatedPerfTestNode.cs b/Engine/Plugins/Performance/AutomatedPerfTesting/Build/Scripts/AutomatedPerfTestNode.cs
--- a/Engine/Plugins/Performance/AutomatedPerfTesting/Build/Scripts/AutomatedPerfTestNode.cs
+++ b/Engine/Plugins/Performance/AutomatedPerfTesting/Build/Scripts/AutomatedPerfTestNode.cs
@@ -78,6 +78,11 @@ namespace AutomatedPerfTest
 			}
 		}
 
+		/// <summary>
+		/// Should we archive dev artifacts even if -dev is passed down.
+		/// </summary>
+		private readonly bool ForceArchiveDevArtifacts = true;
+
 		/// <summary>
 		/// List of errors with special-cased gauntlet messages.
 		/// </summary>
@@ -955,18 +960,23 @@ namespace AutomatedPerfTest
 				OverrideManagedAccountPool(Config.ManagedAcountPool);
 			}
 
-			if(!IsBuildMachine)
+			if(!IsBuildMachine && Config.OverrideArtifactOutputPath)
 			{
 				// If we are on a build machine, artifacts are stored and handled as expected. On build machines, reports
 				// are eventually copied to this location before being processed. When generating local reports, it is
 				// possible artifacts are being output to a place not expected by APT, so we let the runtime know to 
-				// output the artifacts to this location directly. 
+				// output the artifacts to this location directly if required. 
 
 				UnrealTargetPlatform Platform = Context.GetRoleContext(UnrealTargetRole.Client).Platform;
-				string ArtifactOutputPath = GetPerformanceReportArtifactOutputPath(Platform, Config);
-				ClientRole.CommandLineParams.AddUnique("AutomatedPerfTest.ArtifactOutputPath", ArtifactOutputPath);
+				if (PlatformTargetSupport.IsHostMountingSupported(Platform))
+				{
+					string ArtifactOutputPath = GetPerformanceReportArtifactOutputPath(Platform, Config);
+					ClientRole.CommandLineParams.AddUnique("AutomatedPerfTest.ArtifactOutputPath", PlatformTargetSupport.GetHostMountedPath(Platform, ArtifactOutputPath));
+				}
 			}
 
+			ClientRole.ArchiveDevArtifacts = ForceArchiveDevArtifacts;
+
 			return Config;
 		}
 
@@ -1043,7 +1053,7 @@ namespace AutomatedPerfTest
 					Stream.AddArtifact(DateTime.UtcNow, Microsoft.Extensions.Logging.LogLevel.Information, $"{{{ArtifactFiles.Token}}}", ArtifactFiles);
 				}
 			}
-			else
+			else if (Report != null)
 			{
 				foreach (string AbsoluteFilePath in FindFiles(FileFilter, true, OutputPath))
 				{
@@ -1051,6 +1061,10 @@ namespace AutomatedPerfTest
 					Report.AttachArtifact(AbsoluteFilePath, BaseFileName);
 				}
 			}
+			else
+			{
+				Log.Warning("Test Report is null");
+			}
 		}
 
 		public void CopyInsightsTraceToOutput(string FromArtifactPath, string ToOutputPath, out string OutTraceName)
