Jira UE-273381
Changed default file copy over to CommandUtils.Copy since it uses a network safe copy and wait for copy validation

diff --git a/Engine/Plugins/Performance/AutomatedPerfTesting/Build/Scripts/AutomatedPerfTestNode.cs b/Engine/Plugins/Performance/AutomatedPerfTesting/Build/Scripts/AutomatedPerfTestNode.cs
--- a/Engine/Plugins/Performance/AutomatedPerfTesting/Build/Scripts/AutomatedPerfTestNode.cs
+++ b/Engine/Plugins/Performance/AutomatedPerfTesting/Build/Scripts/AutomatedPerfTestNode.cs
@@ -387,10 +387,12 @@ namespace AutomatedPerfTest
 							DataSourceName, ImportDirOverride);
 						ICsvImporter Importer = ReportGenUtils.CreatePerfReportServerImporter(DataSourceName, OriginalBuildName,
 							IsBuildMachine, ImportDirOverride, CommonDataSourceFields);
+
 						if (Importer != null)
 						{
 							// Recursively grab all the csv files we copied to the temp dir and convert them to binary.
 							List<FileInfo> AllBinaryCsvFiles = ReportGenUtils.CollectAndConvertCsvFilesToBinary(OutputPath);
+							Log.Info("Found {0} valid binary csv files for import.", AllBinaryCsvFiles.Count);
 							if (AllBinaryCsvFiles.Count == 0)
 							{
 								throw new AutomationException($"No Csv files found in {OutputPath}");
@@ -410,7 +412,8 @@ namespace AutomatedPerfTest
 
 								// recursively look for trace files that match the CSV's filename in the artifact path
 								string[] MatchingTraces = FindFiles($"*{InsightsFilename}", true, ArtifactPath);
-								if(MatchingTraces.Length > 0)
+
+								if (MatchingTraces.Length > 0)
 								{
 									if (MatchingTraces.Length > 1)
 									{
@@ -569,10 +572,28 @@ namespace AutomatedPerfTest
 						from CsvFile in Directory.GetFiles(CsvsPath, "*.csv", SearchOption.AllDirectories)
 						select CsvFile);
 				}
+				// ->Warn if the directory doesn't exist
+				else
+				{
+					Log.Warning($"Directory does not exist: {CsvsPath}");
+
+				}
 			}
 
+			// ->If no csv files were found, we can't generate a report
 			if (DiscoveredCsvs.Count == 0)
 			{
+				Log.Error("No CSV files found. Searched paths: {0}", string.Join(", ", CsvsPaths));
+				if (Directory.Exists(OutputPath))
+				{
+					var allFiles = Directory.GetFiles(OutputPath, "*.*", SearchOption.AllDirectories);
+					Log.Info($"All files in OutputPath ({OutputPath}): {string.Join(", ", allFiles)}");
+				}
+				else
+				{
+					Log.Warning($"OutputPath does not exist: {OutputPath}");
+				}
+
 				Log.Error("Test completed successfully but no csv profiling results were found. Searched paths were:\r\n  {Paths}", string.Join("\r\n  ", CsvsPaths.Select(s => $"\"{s}\"")));
 				return false;
 			}
@@ -713,6 +734,9 @@ namespace AutomatedPerfTest
 				CSVDirectory.Create();
 			}
 
+			// ->Log to check if dirs exist
+			Log.Info($"OutputDirectory exists: {OutputDirectory.Exists}, CSVDirectory exists: {CSVDirectory.Exists}");
+
 			string ClientArtifactDir = Path.Combine(FromArtifactPath, "Client");
 			string ClientLogPath = Path.Combine(ClientArtifactDir, "ClientOutput.log");
 
@@ -735,7 +759,7 @@ namespace AutomatedPerfTest
 				string DestLogFile = $"CL{Context.BuildInfo.Changelist}-Pass{GetCurrentPass()}-{Guid.Substring(0, Guid.IndexOf("-"))}-{LogFileInfo.Name}";
 				string LogDestPath = Path.Combine(PassDir, DestLogFile);
 				Log.Info("Copying Log {ClientLogPath} To {LogDest}", ClientLogPath, LogDestPath);
-				LogFileInfo.CopyTo(LogDestPath, true);
+				CommandUtils.CopyFile(ClientLogPath, LogDestPath, false, true);
 			}
 			else
 			{
@@ -772,7 +796,8 @@ namespace AutomatedPerfTest
 			string OutputCSVFile = $"{CSVFile.Name.Replace(Extension, "")}-Pass{GetCurrentPass()}{Extension}";
 			string CsvDestPath = Path.Combine(CSVDirectory.FullName, OutputCSVFile);
 			Log.Info("Copying Csv {CsvPath} To {CsvDestPath}", CSVFile.FullName, CsvDestPath);
-			CSVFile.CopyTo(CsvDestPath, true);
+			//CSVFile.CopyTo(CsvDestPath, true);
+			CommandUtils.CopyFile(CSVFile.FullName, CsvDestPath, false, true);
 		}
 		
 		protected virtual string GetSubtestName(TConfigClass Config)
