[APT] PRS Report Upload Link Only shows on Insights Step

#rnx

diff --git a/Engine/Plugins/Performance/AutomatedPerfTesting/Build/Scripts/AutomatedPerfTestNode.cs b/Engine/Plugins/Performance/AutomatedPerfTesting/Build/Scripts/AutomatedPerfTestNode.cs
--- a/Engine/Plugins/Performance/AutomatedPerfTesting/Build/Scripts/AutomatedPerfTestNode.cs
+++ b/Engine/Plugins/Performance/AutomatedPerfTesting/Build/Scripts/AutomatedPerfTestNode.cs
@@ -139,6 +139,11 @@ namespace AutomatedPerfTest
 		/// </summary>
 		private bool GeneratePRSReport = false;
 
+		/// <summary>
+		/// To ensure that PRS report is added only once per report
+		/// </summary>
+		private bool IsPRSAdded = false;
+
 		/// <summary>
 		/// Set up the base list of possible expected errors, plus the messages to deliver if encountered.
 		/// </summary>
@@ -490,10 +495,15 @@ namespace AutomatedPerfTest
 				string UtraceLocation = Path.Combine(OutputPath, "Traces");
 				string FileFilter = "*.utrace";
 				string ArtifactReportName = "Trace Insights Report";
-				string PRSReportName = "PRS Report Upload";
 
 				AddTestReportArtifacts(Result, Report, ArtifactReportName, UtraceLocation, FileFilter);
+			}
+
+			if (!IsPRSAdded)
+			{
+				string PRSReportName = "PRS Report Upload";
 				AddLinkToReport(Report, PRSReportName, PRSReportLink);
+				IsPRSAdded = true;
 			}
 
 			Report.SetMetadata("PerfType", GetSubtestName(Config));
diff --git a/Engine/Source/Programs/AutomationTool/Gauntlet/Framework/Horde/Gauntlet.Horde.cs b/Engine/Source/Programs/AutomationTool/Gauntlet/Framework/Horde/Gauntlet.Horde.cs
--- a/Engine/Source/Programs/AutomationTool/Gauntlet/Framework/Horde/Gauntlet.Horde.cs
+++ b/Engine/Source/Programs/AutomationTool/Gauntlet/Framework/Horde/Gauntlet.Horde.cs
@@ -53,6 +53,13 @@ namespace Gauntlet
 		/// <returns></returns>
 		public static SummaryLineDelegate OnStepSummary;
 
+		/// <summary>
+		/// Delegate called when the Horde step summary is finalized.
+		/// Returned string is appended as a line to the summary report.
+		/// </summary>
+		/// <returns></returns>
+		public static SummaryLineDelegate OnStepSummaryFinalization;
+
 		public static void AddJobBuildMetaData(IBuildSource buildSource)
 		{
 			try
@@ -100,6 +107,14 @@ namespace Gauntlet
 					}
 				}
 
+				if (OnStepSummaryFinalization != null)
+				{
+					foreach (SummaryLineDelegate callback in OnStepSummaryFinalization.GetInvocationList())
+					{
+						Markdown.AppendLine(callback());
+					}
+				}
+
 				// detect if any testdata was generated
 				if (Globals.Params.ParseParam("UseTestDataV2"))
 				{
diff --git a/Engine/Source/Programs/AutomationTool/Gauntlet/Unreal/Base/Gauntlet.UnrealTestNode.cs b/Engine/Source/Programs/AutomationTool/Gauntlet/Unreal/Base/Gauntlet.UnrealTestNode.cs
--- a/Engine/Source/Programs/AutomationTool/Gauntlet/Unreal/Base/Gauntlet.UnrealTestNode.cs
+++ b/Engine/Source/Programs/AutomationTool/Gauntlet/Unreal/Base/Gauntlet.UnrealTestNode.cs
@@ -884,7 +884,7 @@ namespace Gauntlet
 				}
 			}
 
-			Horde.OnStepSummary += () => $"* **{ReportName}**: [Open link]({URLLink})";
+			Horde.OnStepSummaryFinalization += () => $"* **{ReportName}**: [Open link]({URLLink})";
 		}
 
 		public virtual void PopulateCommandlineInfo()
