Update APT tests to default to run with clearPSODriverCache. This reduces run-to-run perftest variability and tests a cold cache which represents a worse case.

#FailedValidations Preflight

diff --git a/Engine/Plugins/Performance/AutomatedPerfTesting/Build/Inc/AutomatedPerfTestCommonSettings.xml b/Engine/Plugins/Performance/AutomatedPerfTesting/Build/Inc/AutomatedPerfTestCommonSettings.xml
--- a/Engine/Plugins/Performance/AutomatedPerfTesting/Build/Inc/AutomatedPerfTestCommonSettings.xml
+++ b/Engine/Plugins/Performance/AutomatedPerfTesting/Build/Inc/AutomatedPerfTestCommonSettings.xml
@@ -41,6 +41,7 @@
 	<Option Name="TestBuildVersion" DefaultValue="$(APTTestBuildVersion)" Description="Overrides test build version"/>
 	<Option Name="ExtraAutomatedPerformanceCommonArgs" DefaultValue="" Description="Additional Common arguments passed to Gauntlet command line"/>
 	<Option Name="NullOSSPlatforms" DefaultValue="Android" Description=""/>
+	<Option Name="ClearShaderCache" Restrict="true|false" DefaultValue="true" Description="Clear the shader cache to test the worst case and improve test consistency."/>
 
 	<!-- Override build version with given value. If this value is empty, APT will use default build version-->
 	<Property Name="BuildVersion" Value="$(APTBuildVersion)"/> 
@@ -78,6 +79,7 @@
 	<Property Name="AutomatedPerformanceCommonArgs" Value="$(AutomatedPerformanceCommonArgs) -AutomatedPerfTest.TestBuildVersion=&quot;$(TestBuildVersion)&quot;"/>
 	<Property Name="AutomatedPerformanceCommonArgs" Value="$(AutomatedPerformanceCommonArgs) -AutomatedPerfTest.BuildVersion=&quot;$(BuildVersion)&quot;"/>
 	<Property Name="AutomatedPerformanceCommonArgs" Value="$(AutomatedPerformanceCommonArgs) -AutomatedPerfTest.NullOSSPlatforms=&quot;$(NullOSSPlatforms)&quot;" If="'$(NullOSSPlatforms)' != ''"/>
+	<Property Name="AutomatedPerformanceCommonArgs" Value="$(AutomatedPerformanceCommonArgs) -AutomatedPerfTest.ClearShaderCache" If="$(ClearShaderCache)"/>
 
 	<Property Name="PerfTestCommonArgs"     Value="$(AutomatedPerformanceCommonArgs) -iterations=$(IterationsPerf)"/>
 	<Property Name="LLMTestCommonArgs"      Value="$(AutomatedPerformanceCommonArgs) -iterations=$(IterationsLLM)"/>
diff --git a/Engine/Plugins/Performance/AutomatedPerfTesting/Build/Scripts/AutomatedPerfTestConfig.cs b/Engine/Plugins/Performance/AutomatedPerfTesting/Build/Scripts/AutomatedPerfTestConfig.cs
--- a/Engine/Plugins/Performance/AutomatedPerfTesting/Build/Scripts/AutomatedPerfTestConfig.cs
+++ b/Engine/Plugins/Performance/AutomatedPerfTesting/Build/Scripts/AutomatedPerfTestConfig.cs
@@ -100,6 +100,13 @@ namespace AutomatedPerfTest
 		[AutoParamWithNames("BasicWorkspace", "AutomatedPerfTest.GPUReshapeWorkspace")]
 		public string GPUReshapeWorkspace;
 
+		/// <summary>
+		/// Whether to clear the shader cache for platforms which support it. Clearing the shader cache to test the worst case and improve
+		/// test consistency.
+		/// </summary>
+		[AutoParamWithNames(false,"AutomatedPerfTest.ClearShaderCache")]
+		public bool ClearShaderCache;
+
 		/// <summary>
 		/// Whether to lock dynamic resolution. This is applied automatically when GPUPerf test is run.
 		/// </summary>
diff --git a/Engine/Plugins/Performance/AutomatedPerfTesting/Build/Scripts/AutomatedPerfTestNode.cs b/Engine/Plugins/Performance/AutomatedPerfTesting/Build/Scripts/AutomatedPerfTestNode.cs
--- a/Engine/Plugins/Performance/AutomatedPerfTesting/Build/Scripts/AutomatedPerfTestNode.cs
+++ b/Engine/Plugins/Performance/AutomatedPerfTesting/Build/Scripts/AutomatedPerfTestNode.cs
@@ -903,6 +903,13 @@ namespace AutomatedPerfTest
 				ClientRole.CommandLineParams.Add("waitforattach");
 			}
 
+			if (Config.ClearShaderCache)
+			{
+				// Not all platforms support -clearPSODriverCache. Future RHIs which support this concept should use a single flag. This may remain
+				// clearPSODriverCache or may be replaced by a new flag with a less implementation specific name.
+				ClientRole.CommandLineParams.Add("clearPSODriverCache");
+			}
+
 			if (!String.IsNullOrEmpty(Config.NullOSSPlatforms) && Config.NullOSSPlatforms.Split("+").Contains(Context.GetRoleContext(UnrealTargetRole.Client).Platform.ToString()))
 			{
 				ClientRole.CommandLineParams.Add("ini:Engine:[OnlineSubsystem]:DefaultPlatformService=Null");
