PR #13843: Quote path of PerfreportTool.dll when calling it from AutomatedPerfTestNode

diff --git a/Engine/Plugins/Performance/AutomatedPerfTesting/Build/Scripts/AutomatedPerfTestNode.cs b/Engine/Plugins/Performance/AutomatedPerfTesting/Build/Scripts/AutomatedPerfTestNode.cs
--- a/Engine/Plugins/Performance/AutomatedPerfTesting/Build/Scripts/AutomatedPerfTestNode.cs
+++ b/Engine/Plugins/Performance/AutomatedPerfTesting/Build/Scripts/AutomatedPerfTestNode.cs
@@ -658,7 +658,7 @@ namespace AutomatedPerfTest
 			var DetailedReportPath = Path.Combine(ReportPath, "Detailed");
 			string[] CsvGenerationArgs = new[]
 			{
-				 $"{ToolPath.FullName}",
+				 $"\"{ToolPath.FullName}\"",
 				 $"-csvdir \"{NewestDir}\"",
 				 $"-reportType \"{ReportType}\"",
 				 $"-o \"{DetailedReportPath}\"",
