Enable APT jobs to override the ProjectName portion of the DataSourceName, which is in the form "Automation.{ProjectName}.{DataSourceType}".

Previously the project name was directly used directly in the string. This coupled the name of the project and downstream processing dashboards. The APT plugin now supports the flag `-AutomatedPerfTest.DataSourceProjectNameOverride=<value>` to override the ProjectName component.

AutomatedPerfTestCommonSettings.xml previously exposed an option DataSourceName which was not used. This change removes that option and replaces it with DataSourceProjectNameOverride. If DataSourceProjectNameOverride is set, the flag `-AutomatedPerfTest.DataSourceProjectNameOverride="$(DataSourceProjectNameOverride)"` is appended to the launch arguments.

diff --git a/Engine/Plugins/Performance/AutomatedPerfTesting/Build/Inc/AutomatedPerfTestCommonSettings.xml b/Engine/Plugins/Performance/AutomatedPerfTesting/Build/Inc/AutomatedPerfTestCommonSettings.xml
--- a/Engine/Plugins/Performance/AutomatedPerfTesting/Build/Inc/AutomatedPerfTestCommonSettings.xml
+++ b/Engine/Plugins/Performance/AutomatedPerfTesting/Build/Inc/AutomatedPerfTestCommonSettings.xml
@@ -25,7 +25,7 @@
 	<Option Name="DoFPSChart" DefaultValue="false" Restrict="true|false" Description="Whether or not to use FPS Chart to gather and parse performance telemetry." />
 	<Option Name="TraceChannels" DefaultValue="default,screenshot,stats" Description="Which channels to use when tracing with Insights" />
 	<Option Name="DoVideoCapture" Restrict="true|false" DefaultValue="false" Description="Whether or not to capture video during the test" />
-	<Option Name="DataSourceName" DefaultValue="Automation.$(ProjectName).Performance" Description="To which PRS data source should this be uploaded"/>
+	<Option Name="DataSourceProjectNameOverride" DefaultValue="" Description="Override the ProjectName segment of the PRS DataSourceName, which is in the form 'Automation.ProjectName.DataSourceType}'. Values must start with a leter, can use only use alpha numeric characters and _, and must be 20 characters or less. "/>
 	<Option Name="GenerateLocalReports" DefaultValue="false" Restrict="true|false" Description="Whether or not to create a local PerfReportTool report" />
 	<Option Name="ResX" DefaultValue="3840" Restrict="^[0-9]+$" Description="Resolution width to run tests with" />
 	<Option Name="ResY" DefaultValue="2160" Restrict="^[0-9]+$" Description="Resolution height to run tests with" />
@@ -78,6 +78,7 @@
 	<Property Name="AutomatedPerformanceCommonArgs" Value="$(AutomatedPerformanceCommonArgs) -AutomatedPerfTest.TestBuildVersion=&quot;$(TestBuildVersion)&quot;"/>
 	<Property Name="AutomatedPerformanceCommonArgs" Value="$(AutomatedPerformanceCommonArgs) -AutomatedPerfTest.BuildVersion=&quot;$(BuildVersion)&quot;"/>
 	<Property Name="AutomatedPerformanceCommonArgs" Value="$(AutomatedPerformanceCommonArgs) -AutomatedPerfTest.NullOSSPlatforms=&quot;$(NullOSSPlatforms)&quot;" If="'$(NullOSSPlatforms)' != ''"/>
+	<Property Name="AutomatedPerformanceCommonArgs" Value="$(AutomatedPerformanceCommonArgs) -AutomatedPerfTest.DataSourceProjectNameOverride=&quot;$(DataSourceProjectNameOverride)&quot;" If="'$(DataSourceProjectNameOverride)' != ''"/>
 
 	<Property Name="PerfTestCommonArgs"     Value="$(AutomatedPerformanceCommonArgs) -iterations=$(IterationsPerf)"/>
 	<Property Name="LLMTestCommonArgs"      Value="$(AutomatedPerformanceCommonArgs) -iterations=$(IterationsLLM)"/>
diff --git a/Engine/Plugins/Performance/AutomatedPerfTesting/Build/Scripts/AutomatedPerfTestConfig.cs b/Engine/Plugins/Performance/AutomatedPerfTesting/Build/Scripts/AutomatedPerfTestConfig.cs
--- a/Engine/Plugins/Performance/AutomatedPerfTesting/Build/Scripts/AutomatedPerfTestConfig.cs
+++ b/Engine/Plugins/Performance/AutomatedPerfTesting/Build/Scripts/AutomatedPerfTestConfig.cs
@@ -14,10 +14,26 @@ namespace AutomatedPerfTest
 {
 	public class AutomatedPerfTestConfigBase : EpicGameTestConfig
 	{
+		/// <summary>
+		/// Used to override the project name portion of the data source name
+		/// Will be inserted into Automation.{dataSourceProjectName}.{dataSourceType} to construct the full data source name for the test. 
+		/// This is useful when the actual project differs from the name displayed downstream in dashboards.
+		/// 
+		/// Values must follow rules similar to standard ProjectNames: 
+		///  1. Must start with a letter.
+		///  2. Can only use alpha numeric characters and _.
+		///  3. Must be 20 characters or less.
+		///  
+		/// The full string can be overridden with AutomatedPerfTest.DataSourceNameOverride 
+		/// </summary>
+		[AutoParamWithNames("AutomatedPerfTest.DataSourceProjectNameOverride")]
+		public string DataSourceProjectNameOverride = "";
+		
 		/// <summary>
 		/// Used to override the test controller portion of the data source name
 		/// Will be appended to Automation.ProjectName to construct the full data source name for the test
-		/// can be overridden with AutomatedPerfTest.DataSourceNameOverride 
+		/// 
+		/// The full string can be overridden with AutomatedPerfTest.DataSourceNameOverride 
 		/// </summary>
 		[AutoParamWithNames("AutomatedPerfTest.DataSourceTypeOverride")]
 		public string DataSourceTypeOverride = "";
@@ -235,10 +251,22 @@ namespace AutomatedPerfTest
 				Log.Logger.LogError("No DataSourceType or DataSourceTypeOverride has been provided.");
 				return null;
 			}
+
+			// If DataSourceProjectNameOverride is set, ensure it is sufficiently similar to a ProjectName.
+			if (!string.IsNullOrWhiteSpace(DataSourceProjectNameOverride))
+			{
+				string reason;
+				if (!IsProjectNameLike(DataSourceProjectNameOverride, out reason))
+				{
+					Log.Logger.LogError($"Invalid DataSourceProjectNameOverride: {reason}");
+					return null;
+				}
+			}
 			
+			string dataSourceProjectName = string.IsNullOrWhiteSpace(DataSourceProjectNameOverride) ? ProjectName : DataSourceProjectNameOverride;
 			string dataSourceType = string.IsNullOrEmpty(DataSourceTypeOverride) ? DataSourceType : DataSourceTypeOverride;
 			
-			return $"Automation.{ProjectName}.{dataSourceType}";
+			return $"Automation.{dataSourceProjectName}.{dataSourceType}";
 		}
 
 		public override void ApplyToConfig(UnrealAppConfig AppConfig, UnrealSessionRole ConfigRole, IEnumerable<UnrealSessionRole> OtherRoles)
@@ -253,6 +281,57 @@ namespace AutomatedPerfTest
 			// instantiate it. 
 			AppConfig.UsePlatformExtraDebugMemory = UsePlatformExtraDebugMemory;
 		}
+
+
+		/// <summary>
+		///  Validation function to ensure a project name override is valid. This mirrors relevant validation in 
+		///  GameProjectUtils::IsValidProjectFileForCreation valid overrides should follow the same set of constraints.
+		/// </summary>
+		/// <returns></returns>
+		private bool IsProjectNameLike(string ProjectName, out string OutFailReason)
+		{
+			// Mirroring MAX_PROJECT_NAME_LENGTH
+			const int maxProjectNameLength = 20;
+
+			if ( string.IsNullOrEmpty(ProjectName) )
+			{
+				OutFailReason = "You must specify a project name.";
+				return false;
+			}
+
+			if ( ProjectName.Length > maxProjectNameLength )
+			{
+				OutFailReason = $"Project names must not be longer than {maxProjectNameLength} characters.";
+				return false;
+			}
+
+			if ( !char.IsLetter(ProjectName[0]) )
+			{
+				OutFailReason = "Project names must begin with an alphabetic character.";
+				return false;
+			}
+
+			// Constrain allowed characters to alpha numeric and '_' to match NameContainsOnlyLegalCharacters.
+			HashSet<char> illegalNameCharacters = new HashSet<char>();
+			foreach (char c in ProjectName)
+			{
+				if ( char.IsLetterOrDigit(c) || c == '_')
+				{
+					continue;
+				}
+
+				illegalNameCharacters.Add(c);
+			}
+			if ( illegalNameCharacters.Count > 0)
+			{
+				string uniqueCharacters = string.Join("", illegalNameCharacters.OrderBy( c => c));
+				OutFailReason = $"Project names may not contain the following characters: '{uniqueCharacters}'";
+				return false;
+			}
+
+			OutFailReason = "";
+			return true;
+		}
 	}
 
 	public class AutomatedSequencePerfTestConfig : AutomatedPerfTestConfigBase
