Enforces value restrictions for the APT Horde configuration parameter DataSourceProjectNameOverride during Horde job setup.

Adds a Restrict attribute to validate the constraints applied in the C# method AutomatedPerfTestConfigBase.IsProjectNameLike().

diff --git a/Engine/Plugins/Performance/AutomatedPerfTesting/Build/Inc/AutomatedPerfTestCommonSettings.xml b/Engine/Plugins/Performance/AutomatedPerfTesting/Build/Inc/AutomatedPerfTestCommonSettings.xml
--- a/Engine/Plugins/Performance/AutomatedPerfTesting/Build/Inc/AutomatedPerfTestCommonSettings.xml
+++ b/Engine/Plugins/Performance/AutomatedPerfTesting/Build/Inc/AutomatedPerfTestCommonSettings.xml
@@ -25,7 +25,7 @@
 	<Option Name="DoFPSChart" DefaultValue="false" Restrict="true|false" Description="Whether or not to use FPS Chart to gather and parse performance telemetry." />
 	<Option Name="TraceChannels" DefaultValue="default,screenshot,stats" Description="Which channels to use when tracing with Insights" />
 	<Option Name="DoVideoCapture" Restrict="true|false" DefaultValue="false" Description="Whether or not to capture video during the test" />
-	<Option Name="DataSourceProjectNameOverride" DefaultValue="" Description="Override the ProjectName segment of the PRS DataSourceName, which is in the form 'Automation.ProjectName.DataSourceType}'. Values must start with a leter, can use only use alpha numeric characters and _, and must be 20 characters or less. "/>
+	<Option Name="DataSourceProjectNameOverride" DefaultValue="" Restrict="^(|[A-Za-z][A-Za-z0-9_]{0,19})$" Description="Override the ProjectName segment of the PRS DataSourceName, which is in the form 'Automation.ProjectName.DataSourceType}'. Values must start with a leter, can use only use alpha numeric characters and _, and must be 20 characters or less. "/>
 	<Option Name="GenerateLocalReports" DefaultValue="false" Restrict="true|false" Description="Whether or not to create a local PerfReportTool report" />
 	<Option Name="ResX" DefaultValue="3840" Restrict="^[0-9]+$" Description="Resolution width to run tests with" />
 	<Option Name="ResY" DefaultValue="2160" Restrict="^[0-9]+$" Description="Resolution height to run tests with" />
