Enable APT Preflight Test Build ID input when preflight

This fixes issues where the preflight build id wouldn't be utilized for generating the PRS report link.

diff --git a/Engine/Plugins/Performance/AutomatedPerfTesting/Build/Inc/AutomatedPerfTestCommonSettings.xml b/Engine/Plugins/Performance/AutomatedPerfTesting/Build/Inc/AutomatedPerfTestCommonSettings.xml
--- a/Engine/Plugins/Performance/AutomatedPerfTesting/Build/Inc/AutomatedPerfTestCommonSettings.xml
+++ b/Engine/Plugins/Performance/AutomatedPerfTesting/Build/Inc/AutomatedPerfTestCommonSettings.xml
@@ -43,7 +43,8 @@
 	<Option Name="NullOSSPlatforms" DefaultValue="Android" Description=""/>
 
 	<!-- Override build version with given value. If this value is empty, APT will use default build version-->
-	<Property Name="BuildVersion" Value="$(APTBuildVersion)"/> 
+	<Property Name="BuildVersion" Value="$(APTBuildVersion)"/>
+	<Property Name="BuildVersion" Value="$(APTTestBuildVersion)" If="$(IsPreflight)"/>
 
 	<!-- Plugin-provided tests -->
 	<Option Name="WithAPTSequenceTests" DefaultValue="true" Restrict="true|false" Description="Whether to run Auto Perf Test sequence flythrough test." />
