Update AutomatedPerformanceTest generated reports to include I/O stats and graphs.

#Synced-CL 49373796

diff --git a/Engine/Plugins/Performance/AutomatedPerfTesting/Build/Scripts/PerfReport/ReportGraphs.xml b/Engine/Plugins/Performance/AutomatedPerfTesting/Build/Scripts/PerfReport/ReportGraphs.xml
--- a/Engine/Plugins/Performance/AutomatedPerfTesting/Build/Scripts/PerfReport/ReportGraphs.xml
+++ b/Engine/Plugins/Performance/AutomatedPerfTesting/Build/Scripts/PerfReport/ReportGraphs.xml
@@ -130,6 +130,10 @@
 			<statString>fmsglogf/*</statString>
 		</graph>
 
+		<graph title="File I/O" showAverages="1" hideStatPrefix="FileIO/" thickness="8" smooth="0" maxY="0" showEvents="LevelStreaming/*">
+			<statString>FileIO/EDLEventQueueDepth,FileIO/ExistingQueuedPackagesQueueDepth,FileIO/QueuedPackagesQueueDepth,IoDispatcher/PendingIoRequests,FileIO/PendingAsyncLoadingIoRequestsSizeMB,IoDispatcherFileBackend/QueuedFilesystemReadMB</statString>
+		</graph>
+
 		<graph title="Texture Streaming" showAverages="1" hideStatPrefix="texturestreaming/" mainStat="textureStreaming/streamingPool" thickness="4" smooth="0" snapToPeaks="0" maxy="0">
 			<statString>texturestreaming/*</statString>
 		</graph>
diff --git a/Engine/Plugins/Performance/AutomatedPerfTesting/Build/Scripts/PerfReport/ReportTypes.xml b/Engine/Plugins/Performance/AutomatedPerfTesting/Build/Scripts/PerfReport/ReportTypes.xml
--- a/Engine/Plugins/Performance/AutomatedPerfTesting/Build/Scripts/PerfReport/ReportTypes.xml
+++ b/Engine/Plugins/Performance/AutomatedPerfTesting/Build/Scripts/PerfReport/ReportTypes.xml
@@ -78,6 +78,24 @@
 		<summary refName="ObjectActorCounts" type="peak" title="Object And Actor Counts" >
 			<stats>ObjectCount/*, ActorCount/*</stats>
 		</summary>
+
+		<!-- I/O uses a boundedstatvalues summary to transform the raw telemetry into more directly actionable stats -->
+		<summary refName="io" type="boundedstatvalues" title="I/O (whole run)" useUnstrippedCsvStats="0">
+			<column name="% Time Queue Busy" formula="percentoverthreshold" threshold="0" summaryStatName="io_busyPercent" stat="fileio/existingqueuedpackagesqueuedepth">
+				<colourThresholds>10,20,30,100</colourThresholds>
+			</column>
+			<column name="% Time Queue Very Busy" formula="percentoverthreshold" threshold="250" summaryStatName="io_veryBusyPercent" stat="fileio/existingqueuedpackagesqueuedepth">
+				<colourThresholds>5,10,15,50</colourThresholds>
+			</column>
+			<column name="MB/Sec (busy)" formula="sum" multiplier="0.0009765625" filterOutZeros="1" perSecond="1" summaryStatName="io_avgMbPerSecBusy" stat="FileIO/PerFrameKB">
+				<colourThresholds>40,20,10,5</colourThresholds>
+			</column>
+			<column name="MB/Sec" formula="sum" multiplier="0.0009765625" filterOutZeros="0" perSecond="1" summaryStatName="io_avgMbPerSec" stat="FileIO/PerFrameKB"/>
+			<column name="Total Loaded (MB)" formula="sum" multiplier="0.0009765625" summaryStatName="io_totalMB" stat="FileIO/PerFrameKB"/>
+			<column name="Package depth (avg)" formula="average" summaryStatName="io_avgPackageDepth" stat="fileio/existingqueuedpackagesqueuedepth"/>
+			<column name="Package depth busy (avg)" formula="average" filterOutZeros="1" summaryStatName="io_avgPackageDepthBusy" stat="fileio/existingqueuedpackagesqueuedepth"/>
+		</summary>
+
 	</sharedSummaries>
 
 	<!-- Preset summary tables -->
@@ -316,6 +334,7 @@
 						
 			<summaryRef name="psos"/>
 		  	<summaryRef name="pso_framehitches"/>
+			<summaryRef name="io"/>
 			<summaryRef name="textureStreaming"/>
 			<summaryRef name="ObjectActorCounts"/>
 
@@ -334,6 +353,7 @@
 			<graph title="Gamethread Breakdown Giant" budget="16.66" external="1"/>
 
 			<graph title="Ticks Breakdown" budget="200"/>
+			<graph title="File I/O"/>
 			<graph title="Actor Counts" budget="0"/>
 			<graph title="Object Count" />
 
@@ -399,6 +419,8 @@
 				<hitchThresholds>60, 150, 250, 500, 750, 1000, 2000</hitchThresholds>
 			</summary>
 
+			<summaryRef name="io"/>
+
 			<graph title="Stat Unit" budget="16.66" />
 			<graph title="Stat Unit Raw" budget="16.66" />
 			<graph title="Gamethread Breakdown" budget="16.66"/>
@@ -407,6 +429,7 @@
 			<graph title="Gamethread Physics" budget="2"/>
 			<graph title="Gamethread Breakdown Giant" budget="16.66" external="1"/>
 			<graph title="Ticks Breakdown" budget="200"/>
+			<graph title="File I/O"/>
 			<graph title="Actor Counts" budget="0"/>
 			<graph title="Cmds DataLayers" budget="300"/>
 			<graph title="Physical Memory" budget="300"/>
@@ -463,6 +486,7 @@
 			</summary>
 
 
+			<summaryRef name="io"/>
 			<summaryRef name="mapOverlay"/>
 			<!--
 			<summaryRef name="syncLoads"/>
@@ -478,6 +502,7 @@
 			<graph title="Gamethread Physics" budget="4"/>
 			<graph title="Gamethread Breakdown Giant" budget="33.33" external="1"/>
 			<graph title="Ticks Breakdown" budget="200"/>
+			<graph title="File I/O"/>
 			<graph title="Actor Counts" budget="0"/>
 			<graph title="Physical Memory" budget="300"/>
 			<graph title="Cmds DayaLayers" budget="300"/>
