Fixed Automated Performance Testing build xml files referencing LLMTestCommonArgs where they should be referencing GPUPerfTestCommonArgs.

#Synced-CL 49889830

diff --git a/Engine/Plugins/Performance/AutomatedPerfTesting/Build/Inc/AutomatedPerfTestCommonSettings.xml b/Engine/Plugins/Performance/AutomatedPerfTesting/Build/Inc/AutomatedPerfTestCommonSettings.xml
--- a/Engine/Plugins/Performance/AutomatedPerfTesting/Build/Inc/AutomatedPerfTestCommonSettings.xml
+++ b/Engine/Plugins/Performance/AutomatedPerfTesting/Build/Inc/AutomatedPerfTestCommonSettings.xml
@@ -126,7 +126,7 @@
 		<ForEach Name="StaticCameraMapName" Values="$(StaticCameraMapNames)" Separator="+">
 			<Property Name="TargetTestList" Value="$(TargetTestList)+AutomatedPerfTest.StaticCameraTest{-AutomatedPerfTest.StaticCameraPerfTest.MapName=$(StaticCameraMapName) -AutomatedPerfTest.DoCSVProfiler $(PerfTestCommonArgs)}[Static Camera Perf $(StaticCameraMapName)]" If="$(DoPerf)" />
 			<Property Name="TargetTestList" Value="$(TargetTestList)+AutomatedPerfTest.StaticCameraTest{-AutomatedPerfTest.StaticCameraPerfTest.MapName=$(StaticCameraMapName) -AutomatedPerfTest.DoCSVProfiler -AutomatedPerfTest.DoLLM $(LLMTestCommonArgs)}[Static Camera LLM $(StaticCameraMapName)]" If="$(DoLLM)"/>
-			<Property Name="TargetTestList" Value="$(TargetTestList)+AutomatedPerfTest.StaticCameraTest{-AutomatedPerfTest.StaticCameraPerfTest.MapName=$(StaticCameraMapName) -AutomatedPerfTest.DoCSVProfiler -AutomatedPerfTest.DoGPUPerf $(LLMTestCommonArgs)}[Static Camera GPUPerf $(StaticCameraMapName)]" If="$(DoGPUPerf)"/>
+			<Property Name="TargetTestList" Value="$(TargetTestList)+AutomatedPerfTest.StaticCameraTest{-AutomatedPerfTest.StaticCameraPerfTest.MapName=$(StaticCameraMapName) -AutomatedPerfTest.DoCSVProfiler -AutomatedPerfTest.DoGPUPerf $(GPUPerfTestCommonArgs)}[Static Camera GPUPerf $(StaticCameraMapName)]" If="$(DoGPUPerf)"/>
 			<Property Name="TargetTestList" Value="$(TargetTestList)+AutomatedPerfTest.StaticCameraTest{-AutomatedPerfTest.StaticCameraPerfTest.MapName=$(StaticCameraMapName) -AutomatedPerfTest.DoCSVProfiler -AutomatedPerfTest.DoInsightsTrace $(InsightsTestCommonArgs)}[Static Camera Insights $(StaticCameraMapName)]" If="$(DoInsightsTrace)"/>
 			<Property Name="TargetTestList" Value="$(TargetTestList)+AutomatedPerfTest.StaticCameraTest{-AutomatedPerfTest.StaticCameraPerfTest.MapName=$(StaticCameraMapName) -AutomatedPerfTest.DoGPUReshape $(GPUReshapeTestCommonArgs)}[Static Camera GPUReshape $(StaticCameraMapName)]" If="$(DoGPUReshape)"/>
 		</ForEach>
@@ -135,7 +135,7 @@
 	<Do If="$(WithAPTMaterialTest)">
 		<Property Name="TargetTestList" Value="$(TargetTestList)+AutomatedPerfTest.MaterialTest{-AutomatedPerfTest.DoCSVProfiler $(PerfTestCommonArgs)}[Material Perf]" If="$(DoPerf)" />
 		<Property Name="TargetTestList" Value="$(TargetTestList)+AutomatedPerfTest.MaterialTest{-AutomatedPerfTest.DoCSVProfiler -AutomatedPerfTest.DoLLM $(LLMTestCommonArgs)}[Material LLM]" If="$(DoLLM)"/>
-		<Property Name="TargetTestList" Value="$(TargetTestList)+AutomatedPerfTest.MaterialTest{-AutomatedPerfTest.DoCSVProfiler -AutomatedPerfTest.DoGPUPerf $(LLMTestCommonArgs)}[Material GPUPerf]" If="$(DoGPUPerf)"/>
+		<Property Name="TargetTestList" Value="$(TargetTestList)+AutomatedPerfTest.MaterialTest{-AutomatedPerfTest.DoCSVProfiler -AutomatedPerfTest.DoGPUPerf $(GPUPerfTestCommonArgs)}[Material GPUPerf]" If="$(DoGPUPerf)"/>
 		<Property Name="TargetTestList" Value="$(TargetTestList)+AutomatedPerfTest.MaterialTest{-AutomatedPerfTest.DoCSVProfiler -AutomatedPerfTest.DoInsightsTrace $(InsightsTestCommonArgs)}[Material Insights]" If="$(DoInsightsTrace)"/>
 		<Property Name="TargetTestList" Value="$(TargetTestList)+AutomatedPerfTest.MaterialTest{-AutomatedPerfTest.DoCSVProfiler -AutomatedPerfTest.DoGPUReshape $(GPUReshapeTestCommonArgs)}[Material GPUReshape]" If="$(DoGPUReshape)"/>
 	</Do>
@@ -143,7 +143,7 @@
 	<Do If="$(WithAPTReplayTest)">
 		<Property Name="TargetTestList" Value="$(TargetTestList)+AutomatedPerfTest.ReplayTest{-AutomatedPerfTest.ReplayPerfTest.ReplayName=&quot;$(ReplayName)&quot; -AutomatedPerfTest.DoCSVProfiler $(PerfTestCommonArgs)}[Replay Perf]" If="$(DoPerf)" />
 		<Property Name="TargetTestList" Value="$(TargetTestList)+AutomatedPerfTest.ReplayTest{-AutomatedPerfTest.ReplayPerfTest.ReplayName=&quot;$(ReplayName)&quot; -AutomatedPerfTest.DoCSVProfiler -AutomatedPerfTest.DoLLM $(LLMTestCommonArgs)}[Replay LLM]" If="$(DoLLM)"/>
-		<Property Name="TargetTestList" Value="$(TargetTestList)+AutomatedPerfTest.ReplayTest{-AutomatedPerfTest.ReplayPerfTest.ReplayName=&quot;$(ReplayName)&quot; -AutomatedPerfTest.DoCSVProfiler -AutomatedPerfTest.DoGPUPerf $(LLMTestCommonArgs)}[Replay GPUPerf]" If="$(DoGPUPerf)"/>
+		<Property Name="TargetTestList" Value="$(TargetTestList)+AutomatedPerfTest.ReplayTest{-AutomatedPerfTest.ReplayPerfTest.ReplayName=&quot;$(ReplayName)&quot; -AutomatedPerfTest.DoCSVProfiler -AutomatedPerfTest.DoGPUPerf $(GPUPerfTestCommonArgs)}[Replay GPUPerf]" If="$(DoGPUPerf)"/>
 		<Property Name="TargetTestList" Value="$(TargetTestList)+AutomatedPerfTest.ReplayTest{-AutomatedPerfTest.ReplayPerfTest.ReplayName=&quot;$(ReplayName)&quot; -AutomatedPerfTest.DoCSVProfiler -AutomatedPerfTest.DoInsightsTrace $(InsightsTestCommonArgs)}[Replay Insights]" If="$(DoInsightsTrace)"/>
 		<Property Name="TargetTestList" Value="$(TargetTestList)+AutomatedPerfTest.ReplayTest{-AutomatedPerfTest.ReplayPerfTest.ReplayName=&quot;$(ReplayName)&quot; -AutomatedPerfTest.DoGPUReshape $(GPUReshapeTestCommonArgs)}[Replay GPUReshape]" If="$(DoGPUReshape)"/>
 	</Do>
diff --git a/Engine/Plugins/Performance/AutomatedPerfTesting/Build/Inc/AutomatedPerfTestZenStreaming.xml b/Engine/Plugins/Performance/AutomatedPerfTesting/Build/Inc/AutomatedPerfTestZenStreaming.xml
--- a/Engine/Plugins/Performance/AutomatedPerfTesting/Build/Inc/AutomatedPerfTestZenStreaming.xml
+++ b/Engine/Plugins/Performance/AutomatedPerfTesting/Build/Inc/AutomatedPerfTestZenStreaming.xml
@@ -58,7 +58,7 @@
 			<ForEach Name="StaticCameraMapName" Values="$(StaticCameraMapNames)" Separator="+">
 				<Property Name="ZenStreamingTargetTestList" Value="$(ZenStreamingTargetTestList)+AutomatedPerfTest.StaticCameraTest{-AutomatedPerfTest.StaticCameraPerfTest.MapName=$(StaticCameraMapName) -AutomatedPerfTest.DoCSVProfiler $(PerfTestCommonArgs) $(ZenTestArgs)}[Zen Static Camera Perf $(StaticCameraMapName)]" If="$(DoPerf) and $(DoBuildZenStreaming)" />
 				<Property Name="ZenStreamingTargetTestList" Value="$(ZenStreamingTargetTestList)+AutomatedPerfTest.StaticCameraTest{-AutomatedPerfTest.StaticCameraPerfTest.MapName=$(StaticCameraMapName) -AutomatedPerfTest.DoCSVProfiler -AutomatedPerfTest.DoLLM $(LLMTestCommonArgs) $(ZenTestArgs)}[Static Camera LLM $(StaticCameraMapName)]" If="$(DoLLM) and $(DoBuildZenStreaming)"/>
-				<Property Name="ZenStreamingTargetTestList" Value="$(ZenStreamingTargetTestList)+AutomatedPerfTest.StaticCameraTest{-AutomatedPerfTest.StaticCameraPerfTest.MapName=$(StaticCameraMapName) -AutomatedPerfTest.DoCSVProfiler -AutomatedPerfTest.DoGPUPerf $(LLMTestCommonArgs) $(ZenTestArgs)}[Static Camera GPUPerf $(StaticCameraMapName)]" If="$(DoGPUPerf) and $(DoBuildZenStreaming)"/>
+				<Property Name="ZenStreamingTargetTestList" Value="$(ZenStreamingTargetTestList)+AutomatedPerfTest.StaticCameraTest{-AutomatedPerfTest.StaticCameraPerfTest.MapName=$(StaticCameraMapName) -AutomatedPerfTest.DoCSVProfiler -AutomatedPerfTest.DoGPUPerf $(GPUPerfTestCommonArgs) $(ZenTestArgs)}[Static Camera GPUPerf $(StaticCameraMapName)]" If="$(DoGPUPerf) and $(DoBuildZenStreaming)"/>
 				<Property Name="ZenStreamingTargetTestList" Value="$(ZenStreamingTargetTestList)+AutomatedPerfTest.StaticCameraTest{-AutomatedPerfTest.StaticCameraPerfTest.MapName=$(StaticCameraMapName) -AutomatedPerfTest.DoCSVProfiler -AutomatedPerfTest.DoInsightsTrace $(InsightsTestCommonArgs) $(ZenTestArgs)}[Static Camera Insights $(StaticCameraMapName)]" If="$(DoInsightsTrace) and $(DoBuildZenStreaming)"/>			
 				<Property Name="ZenStreamingTargetTestList" Value="$(ZenStreamingTargetTestList)+AutomatedPerfTest.StaticCameraTest{-AutomatedPerfTest.StaticCameraPerfTest.MapName=$(StaticCameraMapName) -AutomatedPerfTest.DoGPUReshape $(GPUReshapeTestCommonArgs) $(ZenTestArgs)}[Static Camera GPUReshape $(StaticCameraMapName)]" If="$(DoGPUReshape) and $(DoBuildZenStreaming)"/>
 			</ForEach>
@@ -67,7 +67,7 @@
 		<Do If="$(WithAPTMaterialTest)">
 			<Property Name="ZenStreamingTargetTestList" Value="$(ZenStreamingTargetTestList)+AutomatedPerfTest.MaterialTest{-AutomatedPerfTest.DoCSVProfiler $(PerfTestCommonArgs) $(ZenTestArgs)}[Zen Material Perf]" If="$(DoPerf) and $(DoBuildZenStreaming)"/>
 			<Property Name="ZenStreamingTargetTestList" Value="$(ZenStreamingTargetTestList)+AutomatedPerfTest.MaterialTest{-AutomatedPerfTest.DoCSVProfiler -AutomatedPerfTest.DoLLM $(LLMTestCommonArgs) $(ZenTestArgs)}[Zen Material LLM]" If="$(DoLLM) and $(DoBuildZenStreaming)"/>
-			<Property Name="ZenStreamingTargetTestList" Value="$(ZenStreamingTargetTestList)+AutomatedPerfTest.MaterialTest{-AutomatedPerfTest.DoCSVProfiler -AutomatedPerfTest.DoGPUPerf $(LLMTestCommonArgs) $(ZenTestArgs)}[Zen Material GPUPerf]" If="$(DoGPUPerf) and $(DoBuildZenStreaming)"/>
+			<Property Name="ZenStreamingTargetTestList" Value="$(ZenStreamingTargetTestList)+AutomatedPerfTest.MaterialTest{-AutomatedPerfTest.DoCSVProfiler -AutomatedPerfTest.DoGPUPerf $(GPUPerfTestCommonArgs) $(ZenTestArgs)}[Zen Material GPUPerf]" If="$(DoGPUPerf) and $(DoBuildZenStreaming)"/>
 			<Property Name="ZenStreamingTargetTestList" Value="$(ZenStreamingTargetTestList)+AutomatedPerfTest.MaterialTest{-AutomatedPerfTest.DoCSVProfiler -AutomatedPerfTest.DoInsightsTrace $(InsightsTestCommonArgs) $(ZenTestArgs)}[Zen Material Insights]" If="$(DoInsightsTrace) and $(DoBuildZenStreaming)"/>
 			<Property Name="ZenStreamingTargetTestList" Value="$(ZenStreamingTargetTestList)+AutomatedPerfTest.MaterialTest{-AutomatedPerfTest.DoCSVProfiler -AutomatedPerfTest.DoGPUReshape $(GPUReshapeTestCommonArgs) $(ZenTestArgs)}[Zen Material GPUReshape]" If="$(DoGPUReshape) and $(DoBuildZenStreaming)"/>
 		</Do>
@@ -75,7 +75,7 @@
 		<Do If="$(WithAPTReplayTest)">
 			<Property Name="ZenStreamingTargetTestList" Value="$(ZenStreamingTargetTestList)+AutomatedPerfTest.ReplayTest{-AutomatedPerfTest.ReplayPerfTest.ReplayName=&quot;$(ReplayName)&quot; -AutomatedPerfTest.DoCSVProfiler $(PerfTestCommonArgs) $(ZenTestArgs)}[Zen Replay Perf]" If="$(DoPerf) and $(DoBuildZenStreaming)"/>
 			<Property Name="ZenStreamingTargetTestList" Value="$(ZenStreamingTargetTestList)+AutomatedPerfTest.ReplayTest{-AutomatedPerfTest.ReplayPerfTest.ReplayName=&quot;$(ReplayName)&quot; -AutomatedPerfTest.DoCSVProfiler -AutomatedPerfTest.DoLLM $(LLMTestCommonArgs) $(ZenTestArgs)}[Zen Replay LLM]" If="$(DoLLM) and $(DoBuildZenStreaming)"/>
-			<Property Name="ZenStreamingTargetTestList" Value="$(ZenStreamingTargetTestList)+AutomatedPerfTest.ReplayTest{-AutomatedPerfTest.ReplayPerfTest.ReplayName=&quot;$(ReplayName)&quot; -AutomatedPerfTest.DoCSVProfiler -AutomatedPerfTest.DoGPUPerf $(LLMTestCommonArgs) $(ZenTestArgs)}[Zen Replay GPUPerf]" If="$(DoGPUPerf) and $(DoBuildZenStreaming)"/>
+			<Property Name="ZenStreamingTargetTestList" Value="$(ZenStreamingTargetTestList)+AutomatedPerfTest.ReplayTest{-AutomatedPerfTest.ReplayPerfTest.ReplayName=&quot;$(ReplayName)&quot; -AutomatedPerfTest.DoCSVProfiler -AutomatedPerfTest.DoGPUPerf $(GPUPerfTestCommonArgs) $(ZenTestArgs)}[Zen Replay GPUPerf]" If="$(DoGPUPerf) and $(DoBuildZenStreaming)"/>
 			<Property Name="ZenStreamingTargetTestList" Value="$(ZenStreamingTargetTestList)+AutomatedPerfTest.ReplayTest{-AutomatedPerfTest.ReplayPerfTest.ReplayName=&quot;$(ReplayName)&quot; -AutomatedPerfTest.DoCSVProfiler -AutomatedPerfTest.DoInsightsTrace $(InsightsTestCommonArgs) $(ZenTestArgs)}[Zen Replay Insights]" If="$(DoInsightsTrace) and $(DoBuildZenStreaming)"/>
 			<Property Name="ZenStreamingTargetTestList" Value="$(ZenStreamingTargetTestList)+AutomatedPerfTest.ReplayTest{-AutomatedPerfTest.ReplayPerfTest.ReplayName=&quot;$(ReplayName)&quot; -AutomatedPerfTest.DoGPUReshape $(GPUReshapeTestCommonArgs) $(ZenTestArgs)}[Zen Replay GPUReshape]" If="$(DoGPUReshape) and $(DoBuildZenStreaming)"/>
 		</Do>
