[APT][Analytics] Integrating Horde Telemetry emissions for PRT data.

#Synced-CL 49930492

diff --git a/Engine/Plugins/Performance/AutomatedPerfTesting/Build/Inc/AutomatedPerfTestCommonSettings.xml b/Engine/Plugins/Performance/AutomatedPerfTesting/Build/Inc/AutomatedPerfTestCommonSettings.xml
--- a/Engine/Plugins/Performance/AutomatedPerfTesting/Build/Inc/AutomatedPerfTestCommonSettings.xml
+++ b/Engine/Plugins/Performance/AutomatedPerfTesting/Build/Inc/AutomatedPerfTestCommonSettings.xml
@@ -21,7 +21,8 @@
 	<Option Name="TestID" DefaultValue="APT" Description="A string to use to identify these performance tests" />
 	<Option Name="DeviceProfileOverride" DefaultValue="" Description="Which device profile to use for this test instead of the default value. Will be prepended by PlatformName." />
 	<Option Name="DoCSVProfiler" DefaultValue="true" Restrict="true|false" Description="Whether or not to use CSVProfiler to gather and parse performance telemetry." />
-	<Option Name="PerfReportServer" DefaultValue="false" Restrict="true|false" Description="Whether or not to use CSVProfiler to gather and parse performance telemetry." />
+	<Option Name="PerfReportServer" DefaultValue="false" Restrict="true|false" Description="Whether or not to upload CSVProfiler data to PRS." />
+	<Option Name="HordeTelemetry" DefaultValue="false" Restrict="true|false" Description="Whether or not to emit CSV performance summary data to Horde." />
 	<Option Name="DoFPSChart" DefaultValue="false" Restrict="true|false" Description="Whether or not to use FPS Chart to gather and parse performance telemetry." />
 	<Option Name="TraceChannels" DefaultValue="default,screenshot,stats" Description="Which channels to use when tracing with Insights" />
 	<Option Name="DoVideoCapture" Restrict="true|false" DefaultValue="false" Description="Whether or not to capture video during the test" />
@@ -68,6 +69,7 @@
 	<Property Name="AutomatedPerformanceCommonArgs" Value="$(AutomatedPerformanceCommonArgs) -AutomatedPerfTest.DoFPSChart" If="$(DoFPSChart)"/>
 
 	<Property Name="AutomatedPerformanceCommonArgs" Value="$(AutomatedPerformanceCommonArgs) -PerfReportServer" If="$(PerfReportServer) And $(DoCSVProfiler)"/>
+	<Property Name="AutomatedPerformanceCommonArgs" Value="$(AutomatedPerformanceCommonArgs) -HordeTelemetry" If="$(HordeTelemetry) And $(DoCSVProfiler)"/>
 	<Property Name="AutomatedPerformanceCommonArgs" Value="$(AutomatedPerformanceCommonArgs) -JobDetails=&quot;$(JobDetails)&quot;" If="'$(JobDetails)' != ''"/>
 	<Property Name="AutomatedPerformanceCommonArgs" Value="$(AutomatedPerformanceCommonArgs) -NoLocalReports" If="!$(GenerateLocalReports) And $(DoCSVProfiler)" />
 	<Property Name="AutomatedPerformanceCommonArgs" Value="$(AutomatedPerformanceCommonArgs) -resX=$(ResX) -resY=$(ResY)"/>
diff --git a/Engine/Plugins/Performance/AutomatedPerfTesting/Build/Scripts/AutomatedPerfTestNode.cs b/Engine/Plugins/Performance/AutomatedPerfTesting/Build/Scripts/AutomatedPerfTestNode.cs
--- a/Engine/Plugins/Performance/AutomatedPerfTesting/Build/Scripts/AutomatedPerfTestNode.cs
+++ b/Engine/Plugins/Performance/AutomatedPerfTesting/Build/Scripts/AutomatedPerfTestNode.cs
@@ -144,6 +144,11 @@ namespace AutomatedPerfTest
 		/// </summary>
 		private bool GeneratePRSReport = false;
 
+		/// <summary>
+		/// If true, the resulting CSV files will be emitted to Horde Telemetry.
+		/// </summary>
+		private bool GenerateHordeTelemetry = false;
+
 		/// <summary>
 		/// To ensure that PRS report is added only once per report
 		/// </summary>
@@ -373,6 +378,7 @@ namespace AutomatedPerfTest
 
 					bool bGeneratedLocalReport = false;
 					bool bGeneratedPRSReport = false;
+					bool bGeneratedHordeTelemetry = false;
 
 					// Local report generation is useful for people conducting tests locally without a centralized server (for A/B testing for instance).
 					// To make it work out of the box, enable by default for non-build machine runs
@@ -404,11 +410,13 @@ namespace AutomatedPerfTest
 							// Recursively grab all the csv files we copied to the temp dir and convert them to binary.
 							List<FileInfo> AllBinaryCsvFiles = ReportGenUtils.CollectAndConvertCsvFilesToBinary(OutputPath);
 							Log.Info("Found {0} valid binary csv files for import.", AllBinaryCsvFiles.Count);
+
 							if (AllBinaryCsvFiles.Count == 0)
 							{
 								throw new AutomationException($"No Csv files found in {OutputPath}");
 							}
 
+
 							// The corresponding log for each csv sits in the same subdirectory as the csv file itself.
 							IEnumerable<CsvImportEntry> ImportEntries = AllBinaryCsvFiles
 								.Select(CsvFile => new CsvImportEntry(CsvFile.FullName, Path.Combine(CsvFile.Directory.FullName, "ClientOutput.log")));
@@ -475,9 +483,14 @@ namespace AutomatedPerfTest
 						}
 					}
 
-					if (!bGeneratedLocalReport && !bGeneratedPRSReport)
+					if (GenerateHordeTelemetry)
 					{
-						Log.Warning("Did not generate neither local nor a PRS report.");
+						bGeneratedHordeTelemetry = GenerateHordePerfReportTelemetry(Platform, OutputPath);
+					}
+
+					if (!bGeneratedLocalReport && !bGeneratedPRSReport && !bGeneratedHordeTelemetry)
+					{
+						Log.Warning("Did not generate a local report, a PRS report, or emit Horde telemetry.");
 					}
 				}
 			}
@@ -754,6 +767,90 @@ namespace AutomatedPerfTest
 			return true;
 		}
 
+		/// <summary>
+		/// Emits performance trend telemetry to Horde via Performance Report Tool.
+		/// </summary>
+		/// <param name="Platform">The platform to emit telemetry for.</param>
+		/// <param name="OutputPath">The output path of the csv data.</param>
+		/// <param name="SummaryTable">The summary table to emit data for. Defaults to KeyStats.</param>
+		/// <returns>True if performance telemetry ran successfully, false otherwise.</returns>
+		private bool GenerateHordePerfReportTelemetry(UnrealTargetPlatform Platform, string OutputPath, string SummaryTable = "KeyStats")
+		{
+			var ReportConfigDir = GetCachedConfiguration().ReportConfigDir;
+			if (string.IsNullOrEmpty(ReportConfigDir))
+			{
+				// default to the report types and graphs provided by APT 
+				ReportConfigDir = Path.Combine(Unreal.EngineDirectory.ToString(), "Plugins", "Performance", "AutomatedPerfTesting", "Build", "Scripts", "PerfReport");
+			}
+
+			var DiscoveredCsvs = new List<string>();
+
+			var CsvsPaths = new[]
+			{
+				Path.Combine(OutputPath, "CSV")
+			};
+
+			foreach (var CsvsPath in CsvsPaths)
+			{
+				if (!Directory.Exists(CsvsPath))
+				{
+					Log.Warning($"Directory does not exist: {CsvsPath}");
+					continue;
+				}
+
+				DiscoveredCsvs.AddRange(Directory.GetFiles(CsvsPath, "*.csv.bin", SearchOption.AllDirectories));
+			}
+
+			if (DiscoveredCsvs.Count == 0)
+			{
+				Log.Info($"No .csv.bin files found, attempting CSV -> CSV.BIN conversion on {OutputPath}...");
+
+				List<FileInfo> GeneratedCsvBins = ReportGenUtils.CollectAndConvertCsvFilesToBinary(OutputPath);
+
+				if (GeneratedCsvBins.Count == 0)
+				{
+					Log.Error($"No Csv files found in {OutputPath}");
+
+					return false;
+				}
+
+				DiscoveredCsvs.AddRange(from csvBin in GeneratedCsvBins select csvBin.FullName);
+			}
+
+			// Find the newest csv file and get its directory
+			// (PerfReportTool will only output cached data in -csvdir mode)
+			var NewestFile =
+				(from CsvFile in DiscoveredCsvs
+				 let Timestamp = File.GetCreationTimeUtc(CsvFile)
+				 orderby Timestamp descending
+				 select CsvFile).First();
+			var NewestDir = Path.GetDirectoryName(NewestFile);
+
+			var PlatformNameFilter = Platform == UnrealTargetPlatform.Win64 ? "Windows" : $"{Platform}";
+
+			string[] CsvGenerationArgs = new[]
+			{
+				 $"-csvdir \"{NewestDir}\"",
+				 $"-summaryTable {SummaryTable}",
+				 $"-metadatafilter platform=\"{PlatformNameFilter}\"",
+				 "-nodetailedreports",
+				 "-csvTable",
+				 "-readAllStats",
+				 "-summaryTableDataToTelemetry"
+			};
+
+			string PerfReportToolArgs = string.Join(' ', CsvGenerationArgs);
+			int ErrorCode = ReportGenUtils.RunPerfReportTool(PerfReportToolArgs, ReportConfigDir);
+
+			if (ErrorCode != 0)
+			{
+				Log.Error($"PerfReportTool returned error code \"{ErrorCode}\" while generating detailed report.");
+				return false;
+			}
+
+			return true;
+		}
+
 		private void CopyPerfFilesToOutputDir(string FromArtifactPath, string ToOutputDir = null)
 		{
 			if(string.IsNullOrEmpty(ToOutputDir))
@@ -1205,6 +1302,11 @@ namespace AutomatedPerfTest
 				GeneratePRSReport = true;
 			}
 
+			if(IsBuildMachine && Globals.Params.ParseParam("HordeTelemetry"))
+			{
+				GenerateHordeTelemetry = true;
+			}
+
 			ClientRole.CommandLineParams.Add("AutomatedPerfTest.DoCSVProfiler");
 			ClientRole.CommandLineParams.Add("csvGpuStats");
 
