AutomatedPerfTestZenStreaming.xml, Line 289 - Using #$(PlatformEditorNodeName) (tag reference) instead of $(PlatformEditorNodeName) (node reference) in Requires attribute

Fixed editor test dependencies by correcting a tag-only prerequisite to a node dependency.
The test node was requiring #$(PlatformEditorNodeName), which only waits on tagged outputs and does not guarantee the editor compile node has executed. For non???content-only projects, no tags are produced, so the prerequisite could be silently ignored, creating a potential race. This change switches the dependency to $(PlatformEditorNodeName), making the ordering explicit and consistent with existing BuildGraph usage.

#Synced-CL 49904268

diff --git a/Engine/Plugins/Performance/AutomatedPerfTesting/Build/Inc/AutomatedPerfTestZenStreaming.xml b/Engine/Plugins/Performance/AutomatedPerfTesting/Build/Inc/AutomatedPerfTestZenStreaming.xml
--- a/Engine/Plugins/Performance/AutomatedPerfTesting/Build/Inc/AutomatedPerfTestZenStreaming.xml
+++ b/Engine/Plugins/Performance/AutomatedPerfTesting/Build/Inc/AutomatedPerfTestZenStreaming.xml
@@ -286,7 +286,7 @@
 								<Property Name="GauntletTestArgs" Value="$(GauntletTestArgs) -adbkeys=&quot;$(AdbKeysPath)&quot;" If="'$(AdbKeysPath)' != '' and '$(TargetPlatform)' == 'Android'"/>
 
 								<!-- 'Requires' should be the package game and the publish folder is used; however some tests may require the editor ie: UE.TargetAutomation -->
-								<Node Name="$(TestNodeName)" Requires="#$(PlatformEditorNodeName);$(ZenTestPrerequisites)" Annotations="$(NodeAnnotations)">
+								<Node Name="$(TestNodeName)" Requires="$(PlatformEditorNodeName);$(ZenTestPrerequisites)" Annotations="$(NodeAnnotations)">
 									<Property Name="ZenHost" Value="$(ZenStoreHost)"/>
 									<Property Name="ZenHost" Value="127.0.0.1" If="$(DoLocalZenTest)"/>
 									<Property Name="ZenPlatform" Value="$(TargetPlatform)"/>
