[ATP][Horde][Analytics] Forwarding data source

#Synced-CL 49930492

diff --git a/Engine/Plugins/Performance/AutomatedPerfTesting/Build/Scripts/AutomatedPerfTestNode.cs b/Engine/Plugins/Performance/AutomatedPerfTesting/Build/Scripts/AutomatedPerfTestNode.cs
--- a/Engine/Plugins/Performance/AutomatedPerfTesting/Build/Scripts/AutomatedPerfTestNode.cs
+++ b/Engine/Plugins/Performance/AutomatedPerfTesting/Build/Scripts/AutomatedPerfTestNode.cs
@@ -485,7 +485,7 @@ namespace AutomatedPerfTest
 
 					if (GenerateHordeTelemetry)
 					{
-						bGeneratedHordeTelemetry = GenerateHordePerfReportTelemetry(Platform, OutputPath);
+						bGeneratedHordeTelemetry = GenerateHordePerfReportTelemetry(Platform, DataSourceName, OutputPath);
 					}
 
 					if (!bGeneratedLocalReport && !bGeneratedPRSReport && !bGeneratedHordeTelemetry)
@@ -771,10 +771,11 @@ namespace AutomatedPerfTest
 		/// Emits performance trend telemetry to Horde via Performance Report Tool.
 		/// </summary>
 		/// <param name="Platform">The platform to emit telemetry for.</param>
+		/// <param name="DataSource">The data source of the telemetry invocation.</param>
 		/// <param name="OutputPath">The output path of the csv data.</param>
 		/// <param name="SummaryTable">The summary table to emit data for. Defaults to KeyStats.</param>
 		/// <returns>True if performance telemetry ran successfully, false otherwise.</returns>
-		private bool GenerateHordePerfReportTelemetry(UnrealTargetPlatform Platform, string OutputPath, string SummaryTable = "KeyStats")
+		private bool GenerateHordePerfReportTelemetry(UnrealTargetPlatform Platform, string DataSource, string OutputPath, string SummaryTable = "KeyStats")
 		{
 			var ReportConfigDir = GetCachedConfiguration().ReportConfigDir;
 			if (string.IsNullOrEmpty(ReportConfigDir))
@@ -836,7 +837,8 @@ namespace AutomatedPerfTest
 				 "-nodetailedreports",
 				 "-csvTable",
 				 "-readAllStats",
-				 "-summaryTableDataToTelemetry"
+				 "-summaryTableDataToTelemetry",
+				 $"-telemetryDataSource \"{DataSource}\""
 			};
 
 			string PerfReportToolArgs = string.Join(' ', CsvGenerationArgs);
