[APT][Horde][Analytics] Updating the command line args to PRT for horde telemetry emissions to prevent erroneous data.

#Synced-CL 50425380

diff --git a/Engine/Plugins/Performance/AutomatedPerfTesting/Build/Scripts/AutomatedPerfTestNode.cs b/Engine/Plugins/Performance/AutomatedPerfTesting/Build/Scripts/AutomatedPerfTestNode.cs
--- a/Engine/Plugins/Performance/AutomatedPerfTesting/Build/Scripts/AutomatedPerfTestNode.cs
+++ b/Engine/Plugins/Performance/AutomatedPerfTesting/Build/Scripts/AutomatedPerfTestNode.cs
@@ -828,15 +828,26 @@ namespace AutomatedPerfTest
 			var NewestDir = Path.GetDirectoryName(NewestFile);
 
 			var PlatformNameFilter = Platform == UnrealTargetPlatform.Win64 ? "Windows" : $"{Platform}";
+			
+			string ReportType = null;
+			GetReportType(out ReportType, out _, out _);
+			Log.Info("Horde Telemetry Emission detected the following report type: {0}", ReportType);
+
+			if(ReportType != "ClientPerf" && SummaryTable == "KeyStats")
+			{
+				Log.Info("Horde Telemetry is currently not possible for events outside of ClientPerf reports (when using KeyStats summary table type) - skipping emissions.");
+
+				return true;
+			}
 
 			string[] CsvGenerationArgs = new[]
 			{
 				 $"-csvdir \"{NewestDir}\"",
+				 $"-reportType \"{ReportType}\"",
 				 $"-summaryTable {SummaryTable}",
 				 $"-metadatafilter platform=\"{PlatformNameFilter}\"",
 				 "-nodetailedreports",
 				 "-csvTable",
-				 "-readAllStats",
 				 "-summaryTableDataToTelemetry",
 				 $"-telemetryDataSource \"{DataSource}\""
 			};
