[APT][Analytics] Introduce test identity concept to perf node in order to disambiguate data source into logical groups.

#Synced-CL 50425380

diff --git a/Engine/Plugins/Performance/AutomatedPerfTesting/Build/Scripts/AutomatedPerfTestNode.cs b/Engine/Plugins/Performance/AutomatedPerfTesting/Build/Scripts/AutomatedPerfTestNode.cs
--- a/Engine/Plugins/Performance/AutomatedPerfTesting/Build/Scripts/AutomatedPerfTestNode.cs
+++ b/Engine/Plugins/Performance/AutomatedPerfTesting/Build/Scripts/AutomatedPerfTestNode.cs
@@ -485,7 +485,9 @@ namespace AutomatedPerfTest
 
 					if (GenerateHordeTelemetry)
 					{
-						bGeneratedHordeTelemetry = GenerateHordePerfReportTelemetry(Platform, DataSourceName, OutputPath);
+						string TestIdentity = CreateTestIdentity();
+
+						bGeneratedHordeTelemetry = GenerateHordePerfReportTelemetry(Platform, TestIdentity, OutputPath);
 					}
 
 					if (!bGeneratedLocalReport && !bGeneratedPRSReport && !bGeneratedHordeTelemetry)
@@ -552,8 +554,17 @@ namespace AutomatedPerfTest
 
 		/// <inheritdoc/>
 		protected override string HordeReportTestKey => base.HordeReportTestKey + (string.IsNullOrEmpty(GetCachedConfiguration().TestID) ? "" : $".{GetCachedConfiguration().TestID}");
-		
-		
+
+		/// <summary>
+		/// Gets the test identity of the <see cref="AutomatedPerfTestNode{TConfigClass}"/>.
+		/// </summary>
+		/// <returns>The test identity of the node.</returns>
+		protected virtual string CreateTestIdentity(TConfigClass Config = null)
+		{
+			Config ??= GetCachedConfiguration();
+
+			return Config.DataSourceName;
+		}
 
 		/// <summary>
 		/// Get report type for current configuration
@@ -771,11 +782,11 @@ namespace AutomatedPerfTest
 		/// Emits performance trend telemetry to Horde via Performance Report Tool.
 		/// </summary>
 		/// <param name="Platform">The platform to emit telemetry for.</param>
-		/// <param name="DataSource">The data source of the telemetry invocation.</param>
+		/// <param name="TestIdentity">The test identity of the telemetry.</param>
 		/// <param name="OutputPath">The output path of the csv data.</param>
 		/// <param name="SummaryTable">The summary table to emit data for. Defaults to KeyStats.</param>
 		/// <returns>True if performance telemetry ran successfully, false otherwise.</returns>
-		private bool GenerateHordePerfReportTelemetry(UnrealTargetPlatform Platform, string DataSource, string OutputPath, string SummaryTable = "KeyStats")
+		private bool GenerateHordePerfReportTelemetry(UnrealTargetPlatform Platform, string TestIdentity, string OutputPath, string SummaryTable = "KeyStats")
 		{
 			var ReportConfigDir = GetCachedConfiguration().ReportConfigDir;
 			if (string.IsNullOrEmpty(ReportConfigDir))
@@ -849,7 +860,7 @@ namespace AutomatedPerfTest
 				 "-nodetailedreports",
 				 "-csvTable",
 				 "-summaryTableDataToTelemetry",
-				 $"-telemetryDataSource \"{DataSource}\""
+				 $"-telemetryTestIdentity \"{TestIdentity}\""
 			};
 
 			string PerfReportToolArgs = string.Join(' ', CsvGenerationArgs);
@@ -1481,6 +1492,16 @@ namespace AutomatedPerfTest
 		{
 			return GetNormalizedInsightsFileName(CSVFileName, "_Sequence");
 		}
+
+		/// <inheritdoc/>
+		protected override string CreateTestIdentity(TConfigClass Config = null)
+		{
+			Config ??= GetCachedConfiguration();
+
+			string rootTestIdentity = base.CreateTestIdentity(Config);
+
+			return $"{rootTestIdentity}.{Config.MapSequenceComboName}";
+		}
 	}
 
 	/// <summary>
@@ -1636,11 +1657,23 @@ namespace AutomatedPerfTest
 
 			return OutReplayList;
 		}
+
 		protected override string GetNormalizedInsightsFileName(string CSVFileName)
 		{
 			return GetNormalizedInsightsFileName(CSVFileName, "_Replay");
 		}
 
+		/// <inheritdoc/>
+		protected override string CreateTestIdentity(TConfigClass Config = null)
+		{
+			Config ??= GetCachedConfiguration();
+
+			string rootTestIdentity = base.CreateTestIdentity(Config);
+			string replayName = string.IsNullOrEmpty(Config.ReplayName) ? string.Empty : Path.GetFileNameWithoutExtension(Config.ReplayName);
+
+			return string.IsNullOrEmpty(replayName) ? rootTestIdentity : $"{rootTestIdentity}.{replayName}";
+		}
+
 		private readonly string ReplayDataSourceName = "ReplayRun";
 		private readonly string ReplayParamName = "AutomatedPerfTest.ReplayPerfTest.ReplayName";
 		private TConfigClass Config;
