Update CitySample builds to make performance test configs opt-in.

Prior to the application of PerfTestConfigs in CitySample, enabling a platform implicitly ran performance tests for that platform. PerfTestConfigs multiply the number of tests. For example, enabling the Win64 platform runs the performance tests twice, once for Win64 with the default RHI and once for Win64 with the Vulkan RHI. This will continue to grow as more PerfTestConfigs are added.

This causes two problems:

1. A user iterating on Win64 may want to run a subset of the performance tests for that platform. Running both slows down their iteration speed.
2. Running more tests than required consumes farm resources.

Different jobs may want different sets of tests enabled by default. By starting with all PerfTestConfigs off, job templates can define a default set and expose further configuration to users.

This change additionally adds the macro 'DisableDefaultTests' to Plugins\Performance\AutomatedPerfTesting\Build\Inc\TestConfigs.xml. This should be used by projects using AddPerfTestConfig if the AddPerfTestConfigs are conditionally enabled via Options to avoid falling back to using default tests.

#Synced-CL 50405523

diff --git a/Engine/Plugins/Performance/AutomatedPerfTesting/Build/Inc/TestConfigs.xml b/Engine/Plugins/Performance/AutomatedPerfTesting/Build/Inc/TestConfigs.xml
--- a/Engine/Plugins/Performance/AutomatedPerfTesting/Build/Inc/TestConfigs.xml
+++ b/Engine/Plugins/Performance/AutomatedPerfTesting/Build/Inc/TestConfigs.xml
@@ -9,4 +9,19 @@
 		<Property Name="TestConfigList" Value="$(TestConfigList)^$(TestConfigName)|$(TestPlatform)|$(ArtifactPostfix)|$(AdditionalTestArguments)|$(AppArguments)|$(AgentType)"/>
 	</Macro>
 
+	<!-- 
+	Disables the default tests added when TestConfigList is empty.
+	
+	The presence of a single AddPerfTestConfig macro implicitly disables the default perf test for all platforms. 
+	Projects using AddPerfTestConfig may guard them with Options which default to off. If all tests are off - if 
+	TestConfigList is empty - the default tests run as a fallback.
+	
+	This macro prevents that fallback by adding a no-op item to the list. This no-op adds a test with an invalid 
+	platform name to TestConfigList. BuildAndTestProject filters by platform name this entry is skipped when processing
+	platforms.
+	-->
+	<Macro Name ="DisableDefaultTests">
+		<Property Name="TestConfigList" Value="$(TestConfigList)^DisableDefaultTests|InvalidPlatformIdentifier||||"/>
+	</Macro>
+
 </BuildGraph>
