APT for Zen Streaming Console fixes

Fixed an issue on consoles where CSVs weren't being generated by swithing over from passing zen args into the client from BG to letting the Test Node handle all client args to prevent argument overwritering.
Additionally fixed an issue where on some consoles it was failing due to being unable to find the uproject file during host mounting.

#Synced-CL 50741860

diff --git a/Engine/Build/Graph/Tasks/Inc/PlatformOptions.xml b/Engine/Build/Graph/Tasks/Inc/PlatformOptions.xml
--- a/Engine/Build/Graph/Tasks/Inc/PlatformOptions.xml
+++ b/Engine/Build/Graph/Tasks/Inc/PlatformOptions.xml
@@ -61,7 +61,10 @@
 <!-- These platforms can't run from loose staged content-->
 <Property Name="PlatformsThatRequirePackaging" Value="Android+IOS" />
 
-<Macro Name="AddPlatform" Arguments="Platform;IsConsole;IsEnabled" OptionalArguments="IsVirtualized;TestWithStageBuildOnly">
+<!-- These platforms require the project file to be present when staged -->
+<Property Name="PlatformsRequiringStagedUproject" Value="" />
+
+<Macro Name="AddPlatform" Arguments="Platform;IsConsole;IsEnabled" OptionalArguments="IsVirtualized;TestWithStageBuildOnly;RequiresStagedUproject">
 	<Property Name="IsVirtualized" Value="('$(IsVirtualized)' != '' and '$(IsVirtualized)' != 'false')"/>
 	<Property Name="TestWithStageBuildOnly" Value="('$(TestWithStageBuildOnly)' != '' and '$(TestWithStageBuildOnly)' != 'false')"/>
 	<Property Name="AllTargetPlatforms" Value="$(AllTargetPlatforms)+$(Platform)" If="!ContainsItem('$(AllTargetPlatforms)', '$(Platform)', '+')"/>
@@ -69,6 +72,7 @@
 	<Property Name="AllVirtualizedPlatforms" Value="$(AllVirtualizedPlatforms)+$(Platform)" If="$(IsVirtualized) And !ContainsItem('$(AllVirtualizedPlatforms)', '$(Platform)', '+')"/>
 	<Property Name="TargetPlatforms" Value="$(TargetPlatforms)+$(Platform)" If="$(IsEnabled) and !ContainsItem('$(TargetPlatforms)', '$(Platform)', '+')"/>
 	<Property Name="PlatformsToSkipPackaging" Value="$(PlatformsToSkipPackaging)+$(Platform)" If="$(TestWithStageBuildOnly) and !ContainsItem('$(PlatformsToSkipPackaging)', '$(Platform)', '+')"/>
+	<Property Name="PlatformsRequiringStagedUproject" Value="$(PlatformsRequiringStagedUproject)+$(Platform)" If="('$(RequiresStagedUproject)' == 'true') and !ContainsItem('$(PlatformsRequiringStagedUproject)', '$(Platform)', '+')"/>
 </Macro>
 	
 <!-- Include any scripts from platform extensions -->
diff --git a/Engine/Plugins/Performance/AutomatedPerfTesting/Build/Inc/AutomatedPerfTestZenStreaming.xml b/Engine/Plugins/Performance/AutomatedPerfTesting/Build/Inc/AutomatedPerfTestZenStreaming.xml
--- a/Engine/Plugins/Performance/AutomatedPerfTesting/Build/Inc/AutomatedPerfTestZenStreaming.xml
+++ b/Engine/Plugins/Performance/AutomatedPerfTesting/Build/Inc/AutomatedPerfTestZenStreaming.xml
@@ -34,7 +34,9 @@
     <Option Name="DefaultZenStreamingTargetTestList" DefaultValue="" Description="Default Zen Streaming Target test list"/>
     <Option Name="ZenStreamingTargetTestList" DefaultValue="$(DefaultZenStreamingTargetTestList)" Description="Zen Streaming tests to run on the target platform(s). Test1+Test2 etc"/>
 
-	<Property Name="ZenTestArgs" Value="-AutomatedPerfTest.DoBuildZenStreaming -AutomatedPerfTest.ZenStoreHost=&quot;$(ZenStoreHost)&quot;"/>
+	<Property Name="ZenTestArgs" Value="-AutomatedPerfTest.DoBuildZenStreaming"/>
+	<Property Name="ZenTestArgs" Value="$(ZenTestArgs) -AutomatedPerfTest.ZenStoreHost=&quot;$(ZenStoreHost)&quot;"/>
+	<Property Name="ZenTestArgs" Value="$(ZenTestArgs) -AutomatedPerfTest.ZenStorePort=&quot;$(ZenStorePort)&quot;"/>
 	<Property Name="ZenTestArgs" Value="$(ZenTestArgs) -AutomatedPerfTest.DoLocalZenTest" If="$(DoLocalZenTest)" />
 
     <!-- Aggregate Stubs -->
@@ -149,6 +151,8 @@
 						<Default>
 							<Property Name="ClientStageArguments" Value="$(ZenBuildArgs) -configuration=$(PackageConfigurations) -skipbuild -skipcook -stage -pak -NoCodeSign" />
 							<Command Name="BuildCookRun" Arguments="$(ClientStageArguments)" />
+							<!-- Copy project file for platforms requiring staged uproject - Required for host mounting -->
+							<Copy From="$(TargetProject)" To="$(ProjectOutputDirectory)/$(StagedPlatformFolder)/$(ProjectName)/$(ProjectName).uproject" If="ContainsItem('$(PlatformsRequiringStagedUproject)', '$(TargetPlatform)', '+')"/>
 						</Default>
 					</Switch>
 				</Node>
@@ -158,15 +162,16 @@
 					<Property Name="TargetDirectory" Value="$(ProjectOutputDirectory)/$(StagedPlatformFolder)"/>
 					<Expand Name="WriteArtifactMetadataFile" Type="staged-build" Platform="$(StagedPlatformFolder)" Configuration="$(PackageConfigurations)" ArtifactDir="$(TargetDirectory)" />
 					<Tag Files="$(TargetDirectory)/..." With="#Published Zen $(ProjectName) $(TargetPlatform) Files"/>
+					<Property Name="ZenArtifactMetadata" Value="cookPlatform=$(StagedPlatformFolder);dashboard-category=Staged Builds - Cloud;$(SharedCloudBuildMetadata)" />
 
 					<CreateCloudArtifact 
 						Name="$(ProjectName)-$(TargetPlatform)-Zen-Staged"
 						BuildGroup="$(BuildVersion)" 
 						Type="staged-build" 
 						BaseDir="$(TargetDirectory)" 
-						Metadata="cookPlatform=$(StagedPlatformFolder);dashboard-category=Staged Builds - Cloud;$(SharedCloudBuildMetadata)" 
+						Metadata="$(ZenArtifactMetadata)"
 						Project="$(TargetProject)" 
-						Platform="$(StagedPlatformFolder)" 
+						Platform="$(StagedPlatformFolder)-Zen" 
 						Files="#Published Zen $(ProjectName) $(TargetPlatform) Files"/>
 					<Delete Files="$(TargetDirectory)/artifactmetadata.json"/>
 				</Node>	 
@@ -225,7 +230,7 @@
 									Type="staged-build" 
 									OutputDir="$(ProjectOutputDirectory)/$(StagedPlatformFolder)" 
 									Project="$(TargetProject)" 
-									Platform="$(StagedPlatformFolder)" 
+									Platform="$(StagedPlatformFolder)-Zen" 
 									BuildGroup="$(BuildVersion)" 
 									AllowMultipleMatches="true"
 								/>
@@ -287,15 +292,13 @@
 
 								<!-- 'Requires' should be the package game and the publish folder is used; however some tests may require the editor ie: UE.TargetAutomation -->
 								<Node Name="$(TestNodeName)" Requires="$(PlatformEditorNodeName);$(ZenTestPrerequisites)" Annotations="$(NodeAnnotations)">
-									<Property Name="ZenHost" Value="$(ZenStoreHost)"/>
-									<Property Name="ZenHost" Value="127.0.0.1" If="$(DoLocalZenTest)"/>
 									<Property Name="ZenPlatform" Value="$(TargetPlatform)"/>
 									<Property Name="ZenPlatform" Value="Windows" If="'$(TargetPlatform)' == 'Win64'"/>
-									<Property Name="ZenArgs" Value="-ZenStoreHost=$(ZenHost) -clientargs=&quot;-ZenStorePort=$(ZenStorePort) -ZenStoreProject=$(SnapshotIdentifier) -ZenStorePlatform=$(ZenPlatform)&quot;" />
+									<Property Name="ZenArgs" Value="-AutomatedPerfTest.ZenStoreProject=&quot;$(SnapshotIdentifier)&quot; -AutomatedPerfTest.ZenStorePlatform=&quot;$(ZenPlatform)&quot;" />
 	
 									<!-- Launch a local Zen Server, import the snapshot, execute test, cleanup snapshot -->
 									<ZenLaunch Project="$(TargetProject)" If="$(DoLocalZenTest)" />
-									<ZenImportSnapshot Project="$(TargetProject)" RemoteZenHost="$(ZenHost)" RemoteZenPort="$(ZenStorePort)" SnapshotDescriptorFile="$(ZenSnapshotDescriptorFile)" DestinationProjectId="$(SnapshotIdentifier)" If="$(DoLocalZenTest)"/>
+									<ZenImportSnapshot Project="$(TargetProject)" RemoteZenHost="$(ZenStoreHost)" RemoteZenPort="$(ZenStorePort)" SnapshotDescriptorFile="$(ZenSnapshotDescriptorFile)" DestinationProjectId="$(SnapshotIdentifier)" If="$(DoLocalZenTest)"/>
 									<Command Name="RunUnreal" Arguments="-project=$(TargetProject) -test=&quot;$(TestNameBeforeParentheses)&quot; -platform=$(TargetPlatform) -configuration=$(TestConfiguration) -target=$(TargetName) $(GauntletBuildArg) $(GauntletTestArgs) $(TestArgumentsWithinParentheses) $(GauntletReportArgs) $(TestConfigArgs) $(ZenArgs)"/>
 								</Node>
 
diff --git a/Engine/Plugins/Performance/AutomatedPerfTesting/Build/Scripts/AutomatedPerfTestConfig.cs b/Engine/Plugins/Performance/AutomatedPerfTesting/Build/Scripts/AutomatedPerfTestConfig.cs
--- a/Engine/Plugins/Performance/AutomatedPerfTesting/Build/Scripts/AutomatedPerfTestConfig.cs
+++ b/Engine/Plugins/Performance/AutomatedPerfTesting/Build/Scripts/AutomatedPerfTestConfig.cs
@@ -122,6 +122,24 @@ namespace AutomatedPerfTest
 		[AutoParamWithNames("", "AutomatedPerfTest.ZenStoreHost")]
 		public string ZenStoreHost;
 
+		/// <summary>
+		/// The port of the Zen Server
+		/// </summary>
+		[AutoParamWithNames("", "AutomatedPerfTest.ZenStorePort")]
+		public string ZenStorePort;
+
+		/// <summary>
+		/// The project identifier for the Zen Server
+		/// </summary>
+		[AutoParamWithNames("", "AutomatedPerfTest.ZenStoreProject")]
+		public string ZenStoreProject;
+
+		/// <summary>
+		/// The platform identifier for the Zen Server
+		/// </summary>
+		[AutoParamWithNames("", "AutomatedPerfTest.ZenStorePlatform")]
+		public string ZenStorePlatform;
+
 		/// <summary>
 		/// Whether or not to run a test against a local Zen Server
 		/// </summary>
diff --git a/Engine/Plugins/Performance/AutomatedPerfTesting/Build/Scripts/AutomatedPerfTestNode.cs b/Engine/Plugins/Performance/AutomatedPerfTesting/Build/Scripts/AutomatedPerfTestNode.cs
--- a/Engine/Plugins/Performance/AutomatedPerfTesting/Build/Scripts/AutomatedPerfTestNode.cs
+++ b/Engine/Plugins/Performance/AutomatedPerfTesting/Build/Scripts/AutomatedPerfTestNode.cs
@@ -1028,16 +1028,18 @@ namespace AutomatedPerfTest
 			if (Config.DoBuildZenStreaming)
 			{
 				Config.TestID = "APTZen";
-				string ZenHost = Globals.Params.ParseValue("ZenStoreHost", "any");
 
-				if (Globals.Params.ParseParam("DoLocalZenTest"))
+				if (Config.DoLocalZenTest)
                 {
-                    ZenHost = UnrealHelpers.GetHostIpAddress();
+                    Config.ZenStoreHost = UnrealHelpers.GetHostIpAddress();
                 }
 
+				// Add Zen Streaming params to Client command.
 				ClientRole.CommandLineParams.Add("DoBuildZenStreaming");
-				ClientRole.CommandLineParams.AddOrAppendParamValue("ZenStoreHost", ZenHost);
-
+				ClientRole.CommandLineParams.AddOrAppendParamValue("ZenStoreHost", Config.ZenStoreHost);
+				ClientRole.CommandLineParams.AddOrAppendParamValue("ZenStorePort", Config.ZenStorePort);
+				ClientRole.CommandLineParams.AddOrAppendParamValue("ZenStoreProject", Config.ZenStoreProject);
+				ClientRole.CommandLineParams.AddOrAppendParamValue("ZenStorePlatform", Config.ZenStorePlatform);
 			}
 			ClientRole.CommandLineParams.AddOrAppendParamValue("AutomatedPerfTest.TestID", Config.TestID);
 
