Fix an automated performance test configuration issue for material tests with DoGPUReshape.

The flag -AutomatedPerfTest.DoCSVProfiler which enables the CSVProfiler for measurement was inadverently added to the WithAPTMaterialTest / DoGPUReshape during a past refactor.

#Synced-CL 51220177

diff --git a/Engine/Plugins/Performance/AutomatedPerfTesting/Build/Inc/AutomatedPerfTestCommonSettings.xml b/Engine/Plugins/Performance/AutomatedPerfTesting/Build/Inc/AutomatedPerfTestCommonSettings.xml
--- a/Engine/Plugins/Performance/AutomatedPerfTesting/Build/Inc/AutomatedPerfTestCommonSettings.xml
+++ b/Engine/Plugins/Performance/AutomatedPerfTesting/Build/Inc/AutomatedPerfTestCommonSettings.xml
@@ -137,7 +137,7 @@
 		<Property Name="TargetTestList" Value="$(TargetTestList)+AutomatedPerfTest.MaterialTest{-AutomatedPerfTest.DoCSVProfiler -AutomatedPerfTest.DoLLM $(LLMTestCommonArgs)}[Material LLM]" If="$(DoLLM)"/>
 		<Property Name="TargetTestList" Value="$(TargetTestList)+AutomatedPerfTest.MaterialTest{-AutomatedPerfTest.DoCSVProfiler -AutomatedPerfTest.DoGPUPerf $(GPUPerfTestCommonArgs)}[Material GPUPerf]" If="$(DoGPUPerf)"/>
 		<Property Name="TargetTestList" Value="$(TargetTestList)+AutomatedPerfTest.MaterialTest{-AutomatedPerfTest.DoCSVProfiler -AutomatedPerfTest.DoInsightsTrace $(InsightsTestCommonArgs)}[Material Insights]" If="$(DoInsightsTrace)"/>
-		<Property Name="TargetTestList" Value="$(TargetTestList)+AutomatedPerfTest.MaterialTest{-AutomatedPerfTest.DoCSVProfiler -AutomatedPerfTest.DoGPUReshape $(GPUReshapeTestCommonArgs)}[Material GPUReshape]" If="$(DoGPUReshape)"/>
+		<Property Name="TargetTestList" Value="$(TargetTestList)+AutomatedPerfTest.MaterialTest{-AutomatedPerfTest.DoGPUReshape $(GPUReshapeTestCommonArgs)}[Material GPUReshape]" If="$(DoGPUReshape)"/>
 	</Do>
 
 	<Do If="$(WithAPTReplayTest)">
