[APT][ReportTypes] Adding Frametime

#Synced-CL 51285544

diff --git a/Engine/Plugins/Performance/AutomatedPerfTesting/Build/Scripts/PerfReport/ReportTypes.xml b/Engine/Plugins/Performance/AutomatedPerfTesting/Build/Scripts/PerfReport/ReportTypes.xml
--- a/Engine/Plugins/Performance/AutomatedPerfTesting/Build/Scripts/PerfReport/ReportTypes.xml
+++ b/Engine/Plugins/Performance/AutomatedPerfTesting/Build/Scripts/PerfReport/ReportTypes.xml
@@ -126,7 +126,7 @@
 		</summaryTable>
 
 		<summaryTable name="keystats" rowSort="buildversion,Platform">
-			<filter>Platform,MVP*,Hitches/Min,HitchTimePercent,GamethreadTime*,DynamicResolution*,Renderthreadtime*,GPUTime*,PhysicalUsedMB Max</filter>
+			<filter>Platform,MVP*,Hitches/Min,HitchTimePercent,FrameTime*,GamethreadTime*,DynamicResolution*,Renderthreadtime*,GPUTime*,PhysicalUsedMB Max</filter>
 		</summaryTable>
 
 		<summaryTable name="staticcamera" rowSort="buildversion,Platform,mapname,cameraname,deviceprofile,scalabilitylevel,TargetFPS">
