Fix for Schduled Lyra APT Zen Streaming Perft Test job producing duplicate node collisions with perf test configs.

* Updated `AutomatedPerfTestZenStreaming.xml` to include a line where the `TestConfigTitle` was actually populated with the TestConfig data.

#Synced-CL 50741860

diff --git a/Engine/Plugins/Performance/AutomatedPerfTesting/Build/Inc/AutomatedPerfTestZenStreaming.xml b/Engine/Plugins/Performance/AutomatedPerfTesting/Build/Inc/AutomatedPerfTestZenStreaming.xml
--- a/Engine/Plugins/Performance/AutomatedPerfTesting/Build/Inc/AutomatedPerfTestZenStreaming.xml
+++ b/Engine/Plugins/Performance/AutomatedPerfTesting/Build/Inc/AutomatedPerfTestZenStreaming.xml
@@ -198,6 +198,8 @@
 				<Property Name="TestConfigArgs" Value="$(TestConfigArgs) -args=&quot;$(ConfigAdditionalAppArguments)&quot;" If="'$(ConfigAdditionalAppArguments)' != ''" />
 
 				<Do If="'$(TestConfigTargetPlatform)' == '$(TargetPlatform)'">
+					<Property Name="TestConfigTitle" Value=" $(TestConfigName)" If="'$(TestConfigName)' != ''"/>
+					
 					<Agent Name="Run Zen Streaming Tests $(TargetPlatform)" Type="$(TestTargetAgentType)">
 						<!-- Setup Zen Args -->
 						<Property Name="ZenSnapshotDescriptorFile" Value="$(NetworkSnapshotDescFileRoot)/$(uebp_BuildRoot_Escaped)/$(ProjectName)/$(SnapshotPlatformFolder)/$(Change)/$(SnapshotFilenamePrefix)-builds.json" />
