Android Zen streaming tests in Horde failed silently ? the client exited cleanly (code 0) in ~6 seconds with no output. The StorageServerPlatformFile module could not reach the Zen server and called RequestExit(true).

*  AutomatedPerfTestNode.cs: Skip the `GetHostIpAddress()` override for Android. An empty ZenStoreHost will parse down into localhost.
* BaseEngine.ini: Added `Plugins=AndroidPortForwarder` to [Zen.AutoLaunch]. The plugin config block existed but was never referenced, so it was never loaded.
* AutomatedPerfTestZenStreaming.xml: Added `AndroidPortForwarderDll` copy steps to deploy the DLL through the staged package to test agents, which lack Engine/Source/Developer/.

#Synced-CL 50741860

diff --git a/Engine/Config/BaseEngine.ini b/Engine/Config/BaseEngine.ini
--- a/Engine/Config/BaseEngine.ini
+++ b/Engine/Config/BaseEngine.ini
@@ -2689,6 +2689,7 @@ LocalDataCachePathEnvOverride=UE-LocalDataCachePath
 LocalDataCachePathCommandLineOverride=LocalDataCachePath
 LocalDataCachePathEditorOverrideSetting=LocalDerivedDataCache
 ExtraArgs=--http asio --gc-cache-duration-seconds 1209600 --gc-interval-seconds 21600 --gc-low-diskspace-threshold 2147483648 --cache-bucket-limit-overwrites --quiet
+Plugins=AndroidPortForwarder
 
 [Zen.ConnectExisting]
 HostName=localhost
diff --git a/Engine/Plugins/Performance/AutomatedPerfTesting/Build/Inc/AutomatedPerfTestZenStreaming.xml b/Engine/Plugins/Performance/AutomatedPerfTesting/Build/Inc/AutomatedPerfTestZenStreaming.xml
--- a/Engine/Plugins/Performance/AutomatedPerfTesting/Build/Inc/AutomatedPerfTestZenStreaming.xml
+++ b/Engine/Plugins/Performance/AutomatedPerfTesting/Build/Inc/AutomatedPerfTestZenStreaming.xml
@@ -39,6 +39,9 @@
 	<Property Name="ZenTestArgs" Value="$(ZenTestArgs) -AutomatedPerfTest.ZenStorePort=&quot;$(ZenStorePort)&quot;"/>
 	<Property Name="ZenTestArgs" Value="$(ZenTestArgs) -AutomatedPerfTest.DoLocalZenTest" If="$(DoLocalZenTest)" />
 
+	<!-- Android ADB Port Forwarding required for Host/Client communication when Zen Server is localhosted -->
+	<Property Name="AndroidPortForwarderDll" Value="Engine/Source/Developer/Android/AndroidZenServerPlugin/Binaries/Win64/AndroidPortForwarder2.dll" />
+
     <!-- Aggregate Stubs -->
     <Property Name="GeneratedZenTargetNodes" Value="" />
 	<Property Name="GeneratedZenTargetTestNodes" Value="" />
@@ -136,6 +139,7 @@
 			<Agent Name="Stage and Publish $(PlatformAgentName) Zen" Type="$(CompileAgentType)">
 				<Node Name="Zen Stage $(ProjectName) $(TargetPlatform)" Requires="$(PlatformCompileNodeName);$(PlatformSnapshotPublishNodeName)">
 					<Property Name="ZenBuildArgs" Value="-project=$(TargetProject) -platform=$(TargetPlatform) -ScriptsForProject=$(ProjectName) $(StageAndPackageArguments)" />
+					<Property Name="ZenBuildArgs" Value="$(ZenBuildArgs) -cookflavor=$(AndroidCookFlavor)" If="'$(TargetPlatform)' == 'Android'" />
 					<Property Name="ZenBuildArgs" Value="$(ZenBuildArgs) -client" If="'$(TargetType)' == 'Client'" />			
 
 					<Property Name="ZenSnapshotDescriptorFile" Value="$(NetworkSnapshotDescFileRoot)/$(uebp_BuildRoot_Escaped)/$(ProjectName)/$(SnapshotPlatformFolder)/$(Change)/$(SnapshotFilenamePrefix)-builds.json" />
@@ -145,7 +149,16 @@
 					<Switch>
 						<Case If="$(IsMobilePlatform)">
 							<ForEach Name="PackageConfiguration" Values="$(PackageConfigurations)" Separator="+">
-								<Command Name="BuildCookRun" Arguments="$(ZenBuildArgs) -configuration=$(PackageConfiguration) -skipbuild -skipcook -stage -pak" />
+								<Command Name="BuildCookRun" Arguments="$(ZenBuildArgs) -configuration=$(PackageConfiguration) -skipbuild -skipcook -stage -pak -package" />
+								
+								<!-- Copy APK and Install batch files to staged directory for complete distributable Zen streaming package -->
+								<Do If="'$(TargetPlatform)' == 'Android'">
+									<Copy Files="..." From="$(ProjectPath)/Binaries/Android/$(TargetName)*.apk" To="$(ProjectOutputDirectory)/$(StagedPlatformFolder)" />
+                  					<Copy Files="..." From="$(ProjectPath)/Binaries/Android/Install_$(TargetName)*.bat" To="$(ProjectOutputDirectory)/$(StagedPlatformFolder)" />
+									<Copy Files="..." From="$(ProjectPath)/Binaries/Android/Uninstall_$(TargetName)*.bat" To="$(ProjectOutputDirectory)/$(StagedPlatformFolder)" />
+									<!-- Copy AndroidPortForwarder plugin for Zen server ADB reverse port forwarding -->
+									<Copy From="$(AndroidPortForwarderDll)" To="$(ProjectOutputDirectory)/$(StagedPlatformFolder)/AndroidPortForwarder2.dll" />
+								</Do>
 							</ForEach>
 						</Case>
 						<Default>
@@ -297,6 +310,9 @@
 									<Property Name="ZenPlatform" Value="$(TargetPlatform)"/>
 									<Property Name="ZenPlatform" Value="Windows" If="'$(TargetPlatform)' == 'Win64'"/>
 									<Property Name="ZenArgs" Value="-AutomatedPerfTest.ZenStoreProject=&quot;$(SnapshotIdentifier)&quot; -AutomatedPerfTest.ZenStorePlatform=&quot;$(ZenPlatform)&quot;" />
+
+									<!-- Copy AndroidPortForwarder plugin from staged package for Zen server ADB reverse port forwarding -->
+									<Copy From="$(ProjectOutputDirectory)/$(StagedPlatformFolder)/AndroidPortForwarder2.dll" To="$(AndroidPortForwarderDll)" If="$(DoLocalZenTest) and '$(TargetPlatform)' == 'Android'" />
 	
 									<!-- Launch a local Zen Server, import the snapshot, execute test, cleanup snapshot -->
 									<ZenLaunch Project="$(TargetProject)" If="$(DoLocalZenTest)" />
@@ -328,4 +344,4 @@
 
 		<Aggregate Name="Zen Streaming $(ProjectName)" Requires="$(GeneratedZenTargetNodes);$(GeneratedZenTargetTestNodes)"/>
 	</Do>
-</BuildGraph>
\ No newline at end of file
+</BuildGraph>
diff --git a/Engine/Plugins/Performance/AutomatedPerfTesting/Build/Scripts/AutomatedPerfTestNode.cs b/Engine/Plugins/Performance/AutomatedPerfTesting/Build/Scripts/AutomatedPerfTestNode.cs
--- a/Engine/Plugins/Performance/AutomatedPerfTesting/Build/Scripts/AutomatedPerfTestNode.cs
+++ b/Engine/Plugins/Performance/AutomatedPerfTesting/Build/Scripts/AutomatedPerfTestNode.cs
@@ -1028,10 +1028,15 @@ namespace AutomatedPerfTest
 			if (Config.DoBuildZenStreaming)
 			{
 				Config.TestID = "APTZen";
+				UnrealTargetPlatform ClientPlatform = Context.GetRoleContext(UnrealTargetRole.Client).Platform;
 
 				if (Config.DoLocalZenTest)
                 {
-                    Config.ZenStoreHost = UnrealHelpers.GetHostIpAddress();
+					// Skip Android, empty ZenStoreHost auto-resolves to localhost, allowing AndroidPortForwarder plugin tunnels via adb reverse.
+					if(ClientPlatform != UnrealTargetPlatform.Android)
+					{
+						Config.ZenStoreHost = UnrealHelpers.GetHostIpAddress();	
+					}					
                 }
 
 				// Add Zen Streaming params to Client command.
