[APT] Moved the APT nodes to separate .cs files to make it easier for users to understand what defines a node.

Splitting these into separate files helps keep the code maintainable, reduces cognitive load for new users, and makes it easier to understand what?s involved in adding a new APT node type. The names of the files match the C++ .cpp/.h files, making it easy to open the related set of files.

#Synced-CL 51769453

diff --git a/Engine/Plugins/Performance/AutomatedPerfTesting/Build/Scripts/AutomatedMaterialPerfTest.cs b/Engine/Plugins/Performance/AutomatedPerfTesting/Build/Scripts/AutomatedMaterialPerfTest.cs
new file mode 100644
--- /dev/null
+++ b/Engine/Plugins/Performance/AutomatedPerfTesting/Build/Scripts/AutomatedMaterialPerfTest.cs
@@ -0,0 +1,70 @@
+// Copyright Epic Games, Inc. All Rights Reserved.
+
+using System.Collections.Generic;
+using System.Linq;
+using Gauntlet;
+
+namespace AutomatedPerfTest
+{
+	public class AutomatedMaterialPerfTestConfig : AutomatedPerfTestConfigBase
+	{
+	}
+
+	/// <summary>
+	/// Implementation of a Gauntlet TestNode for AutomatedPerfTest plugin
+	/// </summary>
+	/// <typeparam name="TConfigClass"></typeparam>
+	public abstract class AutomatedMaterialPerfTestNode<TConfigClass> : AutomatedPerfTestNode<TConfigClass>, IAutomatedPerfTest
+		where TConfigClass : AutomatedMaterialPerfTestConfig, new()
+	{
+		public AutomatedMaterialPerfTestNode(UnrealTestContext InContext) : base(InContext)
+		{
+			SummaryTable = "materials";
+		}
+
+		public List<string> GetTestsFromConfig()
+		{
+			List<string> OutMaterials = new List<string>();
+			ReadConfigArray(Context,
+				"/Script/AutomatedPerfTesting.AutomatedMaterialPerfTestProjectSettings",
+				"MaterialsToTest",
+				Material => OutMaterials.Add(Material.Replace("\"", "")));
+
+			return OutMaterials;
+		}
+
+		public override TConfigClass GetConfiguration()
+		{
+			TConfigClass Config = base.GetConfiguration();
+
+			Config.DataSourceName = Config.GetDataSourceName(Context.BuildInfo.ProjectName, "Material");
+
+			// extend the role(s) that we initialized in the base class
+			if (Config.GetRequiredRoles(UnrealTargetRole.Client).Any())
+			{
+				foreach(UnrealTestRole ClientRole in Config.GetRequiredRoles(UnrealTargetRole.Client))
+				{
+					ClientRole.Controllers.Add("AutomatedMaterialPerfTest");
+				}
+			}
+
+			return Config;
+		}
+
+		protected override string GetNormalizedInsightsFileName(string CSVFileName)
+		{
+			return GetNormalizedInsightsFileName(CSVFileName, "_Materials");
+		}
+	}
+
+	/// <summary>
+	/// "Standard issue" implementation usable for samples that don't need anything more advanced
+	/// </summary>
+	public class MaterialTest : AutomatedMaterialPerfTestNode<AutomatedMaterialPerfTestConfig>
+	{
+		public MaterialTest(Gauntlet.UnrealTestContext InContext)
+			: base(InContext)
+		{
+		}
+	}
+}
diff --git a/Engine/Plugins/Performance/AutomatedPerfTesting/Build/Scripts/AutomatedPerfTestConfig.cs b/Engine/Plugins/Performance/AutomatedPerfTesting/Build/Scripts/AutomatedPerfTestConfig.cs
--- a/Engine/Plugins/Performance/AutomatedPerfTesting/Build/Scripts/AutomatedPerfTestConfig.cs
+++ b/Engine/Plugins/Performance/AutomatedPerfTesting/Build/Scripts/AutomatedPerfTestConfig.cs
@@ -375,90 +375,4 @@ namespace AutomatedPerfTest
 			return true;
 		}
 	}
-
-	public class AutomatedSequencePerfTestConfig : AutomatedPerfTestConfigBase
-	{
-		/// <summary>
-		/// Which map to run the test on
-		/// </summary>
-		[AutoParamWithNames("", "AutomatedPerfTest.SequencePerfTest.MapSequenceName")]
-		public string MapSequenceComboName;
-	}
-
-	public class AutomatedReplayPerfTestConfig : AutomatedPerfTestConfigBase
-	{
-		/// <summary>
-		/// Which replay to run the test on
-		/// </summary>
-		[AutoParamWithNames("", "AutomatedPerfTest.ReplayPerfTest.ReplayName")]
-		public string ReplayName;
-	}
-
-	public class AutomatedStaticCameraPerfTestConfig : AutomatedPerfTestConfigBase
-	{
-		/// <summary>
-		/// Which map to run the test on
-		/// </summary>
-		[AutoParamWithNames("", "AutomatedPerfTest.StaticCameraPerfTest.MapName")]
-		public string MapName;
-	}
-	
-	public class AutomatedMaterialPerfTestConfig : AutomatedPerfTestConfigBase
-	{
-	}
-	/// <summary>
-	/// ProfileGo Test config
-	/// </summary>
-	public class ProfileGoConfig 
-	{
-		[AutoParamWithNames("", "ProfileGoCmd")]
-		public string ProfileGoCommand;
-
-		[AutoParamWithNames("", "ProfileGo.Config")]
-		public string ConfigFile;
-
-		[AutoParamWithNames("Default", "ProfileGo")]
-		public string Scenario;
-
-		[AutoParamWithNames("", "ProfileGo.Precmds")]
-		public string PreCmds;
-
-		[AutoParamWithNames("", "ProfileGo.Run")]
-		public string RunCommands;
-
-		[AutoParamWithNames("", "ProfileGo.RunFirst")]
-		public string RunFirstCommands;
-
-		[AutoParamWithNames("", "ProfileGo.RunLast")]
-		public string RunLastCommands;
-
-		[AutoParamWithNames("", "ProfileGo.ExtraArgs")]
-		public string ExtraArgs;
-
-		[AutoParamWithNames(false, "ProfileGo.Uncapped")]
-		public bool bUncapped;
-
-		[AutoParamWithNames(true, "ProfileGo.Exit")]
-		public bool bExit;
-
-		[AutoParamWithNames(-1.0f, "ProfileGo.Settle")]
-		public float Settle;
-
-		[AutoParamWithNames(false, "ProfileGo.SkipSettle")]
-		public bool bSkipSettle;
-
-		[AutoParamWithNames(false, "ProfileGo.RetraceZ")]
-		public bool bRetraceZ;
-
-		[AutoParamWithNames(false, "ProfileGo.DisablePlayerInput")]
-		public bool bDisablePlayerInput;
-	}
-
-	/// <summary>
-	/// ProfileGo Test config
-	/// </summary>
-	public class AutomatedProfileGoTestConfig : AutomatedPerfTestConfigBase
-	{
-		public ProfileGoConfig ProfileGoConfig;
-	}
 }
diff --git a/Engine/Plugins/Performance/AutomatedPerfTesting/Build/Scripts/AutomatedPerfTestNode.cs b/Engine/Plugins/Performance/AutomatedPerfTesting/Build/Scripts/AutomatedPerfTestNode.cs
--- a/Engine/Plugins/Performance/AutomatedPerfTesting/Build/Scripts/AutomatedPerfTestNode.cs
+++ b/Engine/Plugins/Performance/AutomatedPerfTesting/Build/Scripts/AutomatedPerfTestNode.cs
@@ -1447,514 +1447,6 @@ namespace AutomatedPerfTest
 		}
 	}
 
-	/// <summary>
-	/// Implementation of a Gauntlet TestNode for AutomatedPerfTest plugin
-	/// </summary>
-	/// <typeparam name="TConfigClass"></typeparam>
-	public abstract class AutomatedSequencePerfTestNode<TConfigClass> : AutomatedPerfTestNode<TConfigClass>, IAutomatedPerfTest
-		where TConfigClass : AutomatedSequencePerfTestConfig, new()
-	{
-		public AutomatedSequencePerfTestNode(UnrealTestContext InContext) : base(InContext)
-		{
-			SummaryTable = "sequence";
-		}
-
-		public override TConfigClass GetConfiguration()
-		{
-			TConfigClass Config = base.GetConfiguration();
-
-			Config.DataSourceName = Config.GetDataSourceName(Context.BuildInfo.ProjectName, "Sequence");
-			
-			// extend the role(s) that we initialized in the base class
-			if (Config.GetRequiredRoles(UnrealTargetRole.Client).Any())
-			{
-				foreach(UnrealTestRole ClientRole in Config.GetRequiredRoles(UnrealTargetRole.Client))
-				{
-					ClientRole.Controllers.Add("AutomatedSequencePerfTest");
-					
-					// if a specific MapSequenceComboName was defined in the commandline to UAT, then add that to the commandline for the role
-					if (!string.IsNullOrEmpty(Config.MapSequenceComboName))
-					{
-						// use add Unique, since there should only ever be one of these specified
-						ClientRole.CommandLineParams.AddUnique($"AutomatedPerfTest.SequencePerfTest.MapSequenceName",
-							Config.MapSequenceComboName);
-					}
-				}
-			}
-
-			return Config;
-		}
-
-		public List<string> GetTestsFromConfig()
-		{
-			List<string> OutSequenceList = new List<string>();
-			ReadConfigArray(Context,
-				"/Script/AutomatedPerfTesting.AutomatedSequencePerfTestProjectSettings",
-				"MapsAndSequencesToTest",
-				Config =>
-				{
-					Dictionary<string, string> SequenceConfig = IniConfigUtil.ParseDictionaryFromConfigString(Config);
-					string ComboName;
-					if (SequenceConfig.TryGetValue("ComboName", out ComboName))
-					{
-						OutSequenceList.Add(ComboName.Replace("\"", ""));
-					}
-				});
-
-			return OutSequenceList;
-		}
-
-		protected override string GetNormalizedInsightsFileName(string CSVFileName)
-		{
-			return GetNormalizedInsightsFileName(CSVFileName, "_Sequence");
-		}
-
-		/// <inheritdoc/>
-		protected override string CreateTestIdentity(TConfigClass Config = null)
-		{
-			Config ??= GetCachedConfiguration();
-
-			string rootTestIdentity = base.CreateTestIdentity(Config);
-
-			return $"{rootTestIdentity}.{Config.MapSequenceComboName}";
-		}
-	}
-
-	/// <summary>
-	/// Implementation of a Gauntlet TestNode for AutomatedPerfTest plugin
-	/// </summary>
-	/// <typeparam name="TConfigClass"></typeparam>
-	public abstract class AutomatedReplayPerfTestNode<TConfigClass> : AutomatedPerfTestNode<TConfigClass>, IAutomatedPerfTest
-		where TConfigClass : AutomatedReplayPerfTestConfig, new()
-	{
-		public AutomatedReplayPerfTestNode(UnrealTestContext InContext) : base(InContext)
-		{
-			Config = null;
-		}
-
-		/// <summary>
-		/// Handler to explicitly copy files to the device. It is assumed this 
-		/// is called when Gauntlet configures the device before starting the
-		/// tests. 
-		/// </summary>
-		public void CopyReplayToDevice(ITargetDevice Device)
-		{
-			if (Config == null)
-			{
-				Log.Warning("ConfigureDevice() called before Test Node Configuration is set.");
-				return;
-			}
-
-			// At the moment in some pathways we don't seem to copy
-			// files in a role. This ensures we have a copy of the 
-			// replay files on the target device regardless. 
-			if (Config.GetRequiredRoles(UnrealTargetRole.Client).Any())
-			{
-				foreach (UnrealTestRole ClientRole in Config.GetRequiredRoles(UnrealTargetRole.Client))
-				{
-					UnrealTestRoleContext RoleContext = Context.GetRoleContext(UnrealTargetRole.Client);
-					UnrealTargetPlatform Platform = RoleContext.Platform;
-					if (Device.Platform == Platform && ClientRole.FilesToCopy.Any())
-					{
-						UnrealSessionRole Role = new(RoleContext.Type, Platform, RoleContext.Configuration);
-						UnrealAppConfig AppConfig = Context.BuildInfo.CreateConfiguration(Role);
-
-						// needed so AFS will work on CopyAdditional files
-						_ = Device.CreateAppInstall(AppConfig);
-
-						Device.CopyAdditionalFiles(ClientRole.FilesToCopy);
-					}
-				}
-			}
-		}
-
-		public override TConfigClass GetConfiguration()
-		{
-			Config = base.GetConfiguration();
-			
-			Config.DataSourceName = Config.GetDataSourceName(Context.BuildInfo.ProjectName, ReplayDataSourceName);
-
-			IEnumerable<UnrealTestRole> ClientRoles = Config.GetRequiredRoles(UnrealTargetRole.Client);
-			// extend the role(s) that we initialized in the base class
-			if (ClientRoles.Any())
-			{
-				foreach (UnrealTestRole ClientRole in ClientRoles)
-				{
-					ClientRole.Controllers.Add("AutomatedReplayPerfTest");
-					ClientRole.FilesToCopy.Clear();
-
-					UnrealTargetPlatform Platform = Context.GetRoleContext(UnrealTargetRole.Client).Platform;
-					List<string> ReplaysFromConfig = GetTestsFromConfig();
-
-					// Replay name not provided
-					if (string.IsNullOrEmpty(Config.ReplayName))
-					{
-						if(ReplaysFromConfig == null || !ReplaysFromConfig.Any())
-						{
-							throw new AutomationException("No replays found in settings or provided via arguments.");
-						}
-
-						// Use the first replay found. We need to extract the replay path from settings as we need to 
-						// post process the path depending on the target platform.
-						Config.ReplayName = ReplaysFromConfig[0];
-					}
-					else
-					{
-						foreach(string Replay in ReplaysFromConfig)
-						{
-							if(!string.IsNullOrEmpty(Replay) && Replay.Contains(Config.ReplayName))
-							{
-								Log.Info("Found replay in settings");
-								Config.ReplayName = Replay;
-								break;
-							}
-						}
-					}
-
-					string ReplayName = Path.GetFullPath(Config.ReplayName);
-					bool bFileExists = File.Exists(ReplayName);
-					if (!bFileExists)
-					{
-						// In case the replay file specified is not in settings
-						// or not an absolute path. Check under project directory
-						// to be sure.
-						ReplayName = GetPathInProject(Context, Config.ReplayName);
-						bFileExists = File.Exists(ReplayName);
-					}
-
-					if(!bFileExists)
-					{
-						throw new AutomationException($"Replay file '{Config.ReplayName}' not found");
-					}
-
-					if (PlatformTargetSupport.IsHostMountingSupported(Platform)) 
-					{
-						// If current device supports host mounting of files, update the path and use host mounting 
-						ReplayName = PlatformTargetSupport.GetHostMountedPath(Platform, ReplayName);
-						Log.Info($"Host Mountable Platform Detected. Updated Replay File Path to: {ReplayName} ");
-					}
-					else
-					{
-						// Copy replay file to Demos folder in device if host mounting is not supported
-						Log.Info($"Copy Replay File Path to Device: {ReplayName}");
-						UnrealFileToCopy FileToCopy = new UnrealFileToCopy(ReplayName, EIntendedBaseCopyDirectory.Demos, Path.GetFileName(ReplayName));
-						ClientRole.FilesToCopy.Add(FileToCopy);
-						ClientRole.ConfigureDevice = CopyReplayToDevice;
-
-						// If we copy to the "Demos" folder, we just need to pass the Replay Name without
-						// path and extension as the replay subsystem will automatically pick up the file
-						// from this folder.
-						ReplayName = Path.GetFileNameWithoutExtension(ReplayName);
-					}
-
-					ClientRole.CommandLineParams.AddUnique(ReplayParamName, ReplayName);
-					Log.Info($"{ReplayParamName}=\"{ReplayName}\"");
-				}
-			}
-
-			return Config;
-		}
-
-		public List<string> GetTestsFromConfig()
-		{
-			List<string> OutReplayList = new List<string>();
-			ReadConfigArray(Context,
-				"/Script/AutomatedPerfTesting.AutomatedReplayPerfTestProjectSettings",
-				"ReplaysToTest",
-				Config =>
-				{
-					Dictionary<string, string> ReplayConfig = IniConfigUtil.ParseDictionaryFromConfigString(Config);
-					string Path;
-					if (ReplayConfig.TryGetValue("FilePath", out Path))
-					{
-						OutReplayList.Add(GetPathInProject(Context, Path.Replace("\"", "")));
-					}
-				});
-
-			return OutReplayList;
-		}
-
-		protected override string GetNormalizedInsightsFileName(string CSVFileName)
-		{
-			return GetNormalizedInsightsFileName(CSVFileName, "_Replay");
-		}
-
-		/// <inheritdoc/>
-		protected override string CreateTestIdentity(TConfigClass Config = null)
-		{
-			Config ??= GetCachedConfiguration();
-
-			string rootTestIdentity = base.CreateTestIdentity(Config);
-			string replayName = string.IsNullOrEmpty(Config.ReplayName) ? string.Empty : Path.GetFileNameWithoutExtension(Config.ReplayName);
-
-			return string.IsNullOrEmpty(replayName) ? rootTestIdentity : $"{rootTestIdentity}.{replayName}";
-		}
-
-		private readonly string ReplayDataSourceName = "ReplayRun";
-		private readonly string ReplayParamName = "AutomatedPerfTest.ReplayPerfTest.ReplayName";
-		private TConfigClass Config;
-	}
-
-	/// <summary>
-	/// Implementation of a Gauntlet TestNode for AutomatedPerfTest plugin
-	/// </summary>
-	/// <typeparam name="TConfigClass"></typeparam>
-	public abstract class AutomatedStaticCameraPerfTestNode<TConfigClass> : AutomatedPerfTestNode<TConfigClass>, IAutomatedPerfTest
-		where TConfigClass : AutomatedStaticCameraPerfTestConfig, new()
-	{
-		public AutomatedStaticCameraPerfTestNode(UnrealTestContext InContext) : base(InContext)
-		{
-			SummaryTable = "staticcamera";
-		}
-
-		public List<string> GetTestsFromConfig()
-		{
-			List<string> OutMaps = new List<string>();
-			ReadConfigArray(Context,
-				"/Script/AutomatedPerfTesting.AutomatedStaticCameraPerfTestProjectSettings",
-				"MapsToTest", 
-				Map => OutMaps.Add(Map.Replace("\"", "")));
-
-			return OutMaps;
-		}
-
-		public override TConfigClass GetConfiguration()
-		{
-			TConfigClass Config = base.GetConfiguration();
-			
-			Config.DataSourceName = Config.GetDataSourceName(Context.BuildInfo.ProjectName, "StaticCamera");
-			
-			// extend the role(s) that we initialized in the base class
-			if (Config.GetRequiredRoles(UnrealTargetRole.Client).Any())
-			{
-				foreach(UnrealTestRole ClientRole in Config.GetRequiredRoles(UnrealTargetRole.Client))
-				{
-					ClientRole.Controllers.Add("AutomatedPlacedStaticCameraPerfTest");
-					
-					// if a specific MapName was defined in the commandline to UAT, then add that to the commandline for the role
-					if (!string.IsNullOrEmpty(Config.MapName))
-					{
-						// use add Unique, since there should only ever be one of these specified
-						ClientRole.CommandLineParams.AddUnique($"AutomatedPerfTest.StaticCameraPerfTest.MapName",
-							Config.MapName);
-					}
-				}
-			}
-
-			return Config;
-		}
-
-		protected override string GetNormalizedInsightsFileName(string CSVFileName)
-		{
-			return GetNormalizedInsightsFileName(CSVFileName, "_StaticCamera");
-		}
-	}
-
-	/// <summary>
-	/// "Standard issue" implementation usable for samples that don't need anything more advanced
-	/// </summary>
-	public class SequenceTest : AutomatedSequencePerfTestNode<AutomatedSequencePerfTestConfig>
-	{
-		public SequenceTest(Gauntlet.UnrealTestContext InContext)
-			: base(InContext)
-		{
-		}
-	}
-	
-	/// <summary>
-	/// "Standard issue" implementation usable for samples that don't need anything more advanced
-	/// </summary>
-	public class StaticCameraTest : AutomatedStaticCameraPerfTestNode<AutomatedStaticCameraPerfTestConfig>
-	{
-		public StaticCameraTest(Gauntlet.UnrealTestContext InContext)
-			: base(InContext)
-		{
-		}
-	}
-
-	/// <summary>
-	/// Implementation of a Gauntlet TestNode for AutomatedPerfTest plugin
-	/// </summary>
-	/// <typeparam name="TConfigClass"></typeparam>
-	public abstract class AutomatedMaterialPerfTestNode<TConfigClass> : AutomatedPerfTestNode<TConfigClass>, IAutomatedPerfTest
-		where TConfigClass : AutomatedMaterialPerfTestConfig, new()
-	{
-		public AutomatedMaterialPerfTestNode(UnrealTestContext InContext) : base(InContext)
-		{
-			SummaryTable = "materials";
-		}
-
-		public List<string> GetTestsFromConfig()
-		{
-			List<string> OutMaterials = new List<string>();
-			ReadConfigArray(Context,
-				"/Script/AutomatedPerfTesting.AutomatedMaterialPerfTestProjectSettings",
-				"MaterialsToTest",
-				Material => OutMaterials.Add(Material.Replace("\"", "")));
-
-			return OutMaterials;
-		}
-
-		public override TConfigClass GetConfiguration()
-		{
-			TConfigClass Config = base.GetConfiguration();
-			
-			Config.DataSourceName = Config.GetDataSourceName(Context.BuildInfo.ProjectName, "Material");
-			
-			// extend the role(s) that we initialized in the base class
-			if (Config.GetRequiredRoles(UnrealTargetRole.Client).Any())
-			{
-				foreach(UnrealTestRole ClientRole in Config.GetRequiredRoles(UnrealTargetRole.Client))
-				{
-					ClientRole.Controllers.Add("AutomatedMaterialPerfTest");
-				}
-			}
-
-			return Config;
-		}
-
-		protected override string GetNormalizedInsightsFileName(string CSVFileName)
-		{
-			return GetNormalizedInsightsFileName(CSVFileName, "_Materials");
-		}
-	}
-
-	public abstract class AutomatedProfileGoTestNode<TConfigClass> : AutomatedPerfTestNode<TConfigClass>, IAutomatedPerfTest
-		where TConfigClass : AutomatedProfileGoTestConfig, new()
-	{
-		public AutomatedProfileGoTestNode(UnrealTestContext InContext) : base(InContext)
-		{
-			ProfileGoConfig = new ProfileGoConfig();
-			AutoParam.ApplyParamsAndDefaults(ProfileGoConfig, InContext.TestParams.AllArguments);
-		}
-
-		public List<string> GetTestsFromConfig()
-		{
-			List<string> List = new List<string>();
-			return List;
-		}
-
-		public override TConfigClass GetConfiguration()
-		{
-			TConfigClass Config = base.GetConfiguration();
-			string DefaultConfigPath = Path.Combine(Context.BuildInfo.ProjectPath.Directory.FullName, "Build", "ProfileGo", "ProfileGo.json"); 
-			UnrealTargetPlatform Platform = Context.GetRoleContext(UnrealTargetRole.Client).Platform;
-			Config.DataSourceName = Config.GetDataSourceName(Context.BuildInfo.ProjectName, "ProfileGo");
-			if (Config.GetRequiredRoles(UnrealTargetRole.Client).Any())
-			{
-				foreach (UnrealTestRole ClientRole in Config.GetRequiredRoles(UnrealTargetRole.Client))
-				{
-					ClientRole.Controllers.Add("AutomatedProfileGoTest");
-
-					// Copy json config file if available. 
-					string ConfigFile = string.IsNullOrEmpty(ProfileGoConfig.ConfigFile) || !File.Exists(ProfileGoConfig.ConfigFile) ? DefaultConfigPath : ProfileGoConfig.ConfigFile;
-					if (File.Exists(ConfigFile))
-					{
-						if (PlatformTargetSupport.IsHostMountingSupported(Platform))
-						{
-							// If current device supports host mounting of files, update the path and use host mounting 
-							ProfileGoConfig.ConfigFile = PlatformTargetSupport.GetHostMountedPath(Platform, ConfigFile);
-							Log.Info($"Host Mountable Platform Detected. Updated ProfileGo Config File Path to: {ConfigFile} ");
-						}
-						else
-						{
-							Log.Info($"Copying Config File to Device: {ConfigFile}");
-							UnrealFileToCopy FileToCopy = new UnrealFileToCopy(ConfigFile, EIntendedBaseCopyDirectory.Profiling, Path.GetFileName(ConfigFile));
-							ClientRole.FilesToCopy.Add(FileToCopy);
-						}
-					}
-				}
-			}
-
-			ApplyProfileGoParams(Config);
-			Config.ProfileGoConfig = ProfileGoConfig;
-			return Config;
-		}
-
-		protected override string GetNormalizedInsightsFileName(string CSVFileName)
-		{
-			return GetNormalizedInsightsFileName(CSVFileName, "_ProfileGo");
-		}
-
-		protected virtual void ApplyProfileGoParams(TConfigClass Config)
-		{
-			// FYI: ProfileGo in APT is WIP and is subject to change.
-			// Applying default values if none available.
-			if (string.IsNullOrEmpty(ProfileGoConfig.Scenario))
-			{
-				ProfileGoConfig.Scenario = "Default";
-			}
-
-			if (string.IsNullOrEmpty(ProfileGoConfig.RunCommands))
-			{
-				ProfileGoConfig.RunCommands = "profileperfspinning,screenshotnamed,meshstats";
-			}
-
-			foreach (UnrealTestRole ClientRole in Config.RequiredRoles[UnrealTargetRole.Client])
-			{
-				if (string.IsNullOrEmpty(ProfileGoConfig.ConfigFile) == false) { ClientRole.CommandLineParams.Add("profilego.config", ProfileGoConfig.ConfigFile); }
-				if (string.IsNullOrEmpty(ProfileGoConfig.ProfileGoCommand) == false) { ClientRole.CommandLineParams.Add("profilegocmd", ProfileGoConfig.ProfileGoCommand); }
-				if (string.IsNullOrEmpty(ProfileGoConfig.Scenario) == false) { ClientRole.CommandLineParams.Add("profilego", ProfileGoConfig.Scenario); }
-				if (string.IsNullOrEmpty(ProfileGoConfig.PreCmds) == false) { ClientRole.CommandLineParams.Add("profilego.precmds", ProfileGoConfig.PreCmds); }
-				if (string.IsNullOrEmpty(ProfileGoConfig.RunCommands) == false) { ClientRole.CommandLineParams.Add("profilego.run", ProfileGoConfig.RunCommands); }
-				if (string.IsNullOrEmpty(ProfileGoConfig.RunFirstCommands) == false) { ClientRole.CommandLineParams.Add("profilego.runfirst", ProfileGoConfig.RunFirstCommands); }
-				if (string.IsNullOrEmpty(ProfileGoConfig.RunLastCommands) == false) { ClientRole.CommandLineParams.Add("profilego.runlast", ProfileGoConfig.RunLastCommands); }
-				if (string.IsNullOrEmpty(ProfileGoConfig.ExtraArgs) == false) { ClientRole.CommandLineParams.Add("profilego.extraargs", ProfileGoConfig.ExtraArgs); }
-				if (ProfileGoConfig.bUncapped) { ClientRole.CommandLineParams.Add("profilego.uncapped"); }
-				if (ProfileGoConfig.bExit) { ClientRole.CommandLineParams.Add("profilego.exit"); }
-				if (ProfileGoConfig.Settle > 0.0f) { ClientRole.CommandLineParams.Add("profilego.settle", ProfileGoConfig.Settle); }
-				if (ProfileGoConfig.bDisablePlayerInput) { ClientRole.CommandLineParams.Add("profilego.disableplayerinput"); }
-			}
-		}
-
-		protected override bool CopyCSVs(List<FileInfo> CsvFiles, DirectoryInfo CSVDirectory)
-		{
-			// For ProfileGo, each "scenario" can generate a CSV file, so we have
-			// to ensure we copy all valid CSVs to output directory
-			if(CsvFiles.Count == 0)
-			{
-				return false;
-			}
-
-			foreach (FileInfo CsvFile in CsvFiles)
-			{
-				CopyCSV(CsvFile, CSVDirectory);
-			}
-
-			return true;
-		}
-
-		private ProfileGoConfig ProfileGoConfig;
-	}
-
-	public class ProfileGoTest : AutomatedProfileGoTestNode<AutomatedProfileGoTestConfig>
-	{
-		public ProfileGoTest(UnrealTestContext InContext)
-			: base(InContext) { }
-	}
-
-	/// <summary>
-	/// "Standard issue" implementation usable for samples that don't need anything more advanced
-	/// </summary>
-	public class MaterialTest : AutomatedMaterialPerfTestNode<AutomatedMaterialPerfTestConfig>
-	{
-		public MaterialTest(Gauntlet.UnrealTestContext InContext)
-			: base(InContext)
-		{
-		}
-	}
-
-	/// <summary>
-	/// "Standard issue" implementation usable for samples that don't need anything more advanced
-	/// </summary>
-	public class ReplayTest : AutomatedReplayPerfTestNode<AutomatedReplayPerfTestConfig>
-	{
-		public ReplayTest(Gauntlet.UnrealTestContext InContext)
-			: base(InContext)
-		{
-		}
-	}
-
 	public class DefaultTest : AutomatedPerfTestNode<AutomatedPerfTestConfigBase>
 	{
 		public DefaultTest(UnrealTestContext InContext) 
diff --git a/Engine/Plugins/Performance/AutomatedPerfTesting/Build/Scripts/AutomatedProfileGoTest.cs b/Engine/Plugins/Performance/AutomatedPerfTesting/Build/Scripts/AutomatedProfileGoTest.cs
new file mode 100644
--- /dev/null
+++ b/Engine/Plugins/Performance/AutomatedPerfTesting/Build/Scripts/AutomatedProfileGoTest.cs
@@ -0,0 +1,182 @@
+// Copyright Epic Games, Inc. All Rights Reserved.
+
+using System.Collections.Generic;
+using System.IO;
+using System.Linq;
+using Gauntlet;
+using AutomationTool;
+using Log = Gauntlet.Log;
+using UnrealBuildTool;
+
+namespace AutomatedPerfTest
+{
+	/// <summary>
+	/// ProfileGo Test config
+	/// </summary>
+	public class ProfileGoConfig
+	{
+		[AutoParamWithNames("", "ProfileGoCmd")]
+		public string ProfileGoCommand;
+
+		[AutoParamWithNames("", "ProfileGo.Config")]
+		public string ConfigFile;
+
+		[AutoParamWithNames("Default", "ProfileGo")]
+		public string Scenario;
+
+		[AutoParamWithNames("", "ProfileGo.Precmds")]
+		public string PreCmds;
+
+		[AutoParamWithNames("", "ProfileGo.Run")]
+		public string RunCommands;
+
+		[AutoParamWithNames("", "ProfileGo.RunFirst")]
+		public string RunFirstCommands;
+
+		[AutoParamWithNames("", "ProfileGo.RunLast")]
+		public string RunLastCommands;
+
+		[AutoParamWithNames("", "ProfileGo.ExtraArgs")]
+		public string ExtraArgs;
+
+		[AutoParamWithNames(false, "ProfileGo.Uncapped")]
+		public bool bUncapped;
+
+		[AutoParamWithNames(true, "ProfileGo.Exit")]
+		public bool bExit;
+
+		[AutoParamWithNames(-1.0f, "ProfileGo.Settle")]
+		public float Settle;
+
+		[AutoParamWithNames(false, "ProfileGo.SkipSettle")]
+		public bool bSkipSettle;
+
+		[AutoParamWithNames(false, "ProfileGo.RetraceZ")]
+		public bool bRetraceZ;
+
+		[AutoParamWithNames(false, "ProfileGo.DisablePlayerInput")]
+		public bool bDisablePlayerInput;
+	}
+
+	/// <summary>
+	/// ProfileGo Test config
+	/// </summary>
+	public class AutomatedProfileGoTestConfig : AutomatedPerfTestConfigBase
+	{
+		public ProfileGoConfig ProfileGoConfig;
+	}
+
+	public abstract class AutomatedProfileGoTestNode<TConfigClass> : AutomatedPerfTestNode<TConfigClass>, IAutomatedPerfTest
+		where TConfigClass : AutomatedProfileGoTestConfig, new()
+	{
+		public AutomatedProfileGoTestNode(UnrealTestContext InContext) : base(InContext)
+		{
+			ProfileGoConfig = new ProfileGoConfig();
+			AutoParam.ApplyParamsAndDefaults(ProfileGoConfig, InContext.TestParams.AllArguments);
+		}
+
+		public List<string> GetTestsFromConfig()
+		{
+			List<string> List = new List<string>();
+			return List;
+		}
+
+		public override TConfigClass GetConfiguration()
+		{
+			TConfigClass Config = base.GetConfiguration();
+			string DefaultConfigPath = Path.Combine(Context.BuildInfo.ProjectPath.Directory.FullName, "Build", "ProfileGo", "ProfileGo.json");
+			UnrealTargetPlatform Platform = Context.GetRoleContext(UnrealTargetRole.Client).Platform;
+			Config.DataSourceName = Config.GetDataSourceName(Context.BuildInfo.ProjectName, "ProfileGo");
+			if (Config.GetRequiredRoles(UnrealTargetRole.Client).Any())
+			{
+				foreach (UnrealTestRole ClientRole in Config.GetRequiredRoles(UnrealTargetRole.Client))
+				{
+					ClientRole.Controllers.Add("AutomatedProfileGoTest");
+
+					// Copy json config file if available.
+					string ConfigFile = string.IsNullOrEmpty(ProfileGoConfig.ConfigFile) || !File.Exists(ProfileGoConfig.ConfigFile) ? DefaultConfigPath : ProfileGoConfig.ConfigFile;
+					if (File.Exists(ConfigFile))
+					{
+						if (PlatformTargetSupport.IsHostMountingSupported(Platform))
+						{
+							// If current device supports host mounting of files, update the path and use host mounting
+							ProfileGoConfig.ConfigFile = PlatformTargetSupport.GetHostMountedPath(Platform, ConfigFile);
+							Log.Info($"Host Mountable Platform Detected. Updated ProfileGo Config File Path to: {ConfigFile} ");
+						}
+						else
+						{
+							Log.Info($"Copying Config File to Device: {ConfigFile}");
+							UnrealFileToCopy FileToCopy = new UnrealFileToCopy(ConfigFile, EIntendedBaseCopyDirectory.Profiling, Path.GetFileName(ConfigFile));
+							ClientRole.FilesToCopy.Add(FileToCopy);
+						}
+					}
+				}
+			}
+
+			ApplyProfileGoParams(Config);
+			Config.ProfileGoConfig = ProfileGoConfig;
+			return Config;
+		}
+
+		protected override string GetNormalizedInsightsFileName(string CSVFileName)
+		{
+			return GetNormalizedInsightsFileName(CSVFileName, "_ProfileGo");
+		}
+
+		protected virtual void ApplyProfileGoParams(TConfigClass Config)
+		{
+			// FYI: ProfileGo in APT is WIP and is subject to change.
+			// Applying default values if none available.
+			if (string.IsNullOrEmpty(ProfileGoConfig.Scenario))
+			{
+				ProfileGoConfig.Scenario = "Default";
+			}
+
+			if (string.IsNullOrEmpty(ProfileGoConfig.RunCommands))
+			{
+				ProfileGoConfig.RunCommands = "profileperfspinning,screenshotnamed,meshstats";
+			}
+
+			foreach (UnrealTestRole ClientRole in Config.RequiredRoles[UnrealTargetRole.Client])
+			{
+				if (string.IsNullOrEmpty(ProfileGoConfig.ConfigFile) == false) { ClientRole.CommandLineParams.Add("profilego.config", ProfileGoConfig.ConfigFile); }
+				if (string.IsNullOrEmpty(ProfileGoConfig.ProfileGoCommand) == false) { ClientRole.CommandLineParams.Add("profilegocmd", ProfileGoConfig.ProfileGoCommand); }
+				if (string.IsNullOrEmpty(ProfileGoConfig.Scenario) == false) { ClientRole.CommandLineParams.Add("profilego", ProfileGoConfig.Scenario); }
+				if (string.IsNullOrEmpty(ProfileGoConfig.PreCmds) == false) { ClientRole.CommandLineParams.Add("profilego.precmds", ProfileGoConfig.PreCmds); }
+				if (string.IsNullOrEmpty(ProfileGoConfig.RunCommands) == false) { ClientRole.CommandLineParams.Add("profilego.run", ProfileGoConfig.RunCommands); }
+				if (string.IsNullOrEmpty(ProfileGoConfig.RunFirstCommands) == false) { ClientRole.CommandLineParams.Add("profilego.runfirst", ProfileGoConfig.RunFirstCommands); }
+				if (string.IsNullOrEmpty(ProfileGoConfig.RunLastCommands) == false) { ClientRole.CommandLineParams.Add("profilego.runlast", ProfileGoConfig.RunLastCommands); }
+				if (string.IsNullOrEmpty(ProfileGoConfig.ExtraArgs) == false) { ClientRole.CommandLineParams.Add("profilego.extraargs", ProfileGoConfig.ExtraArgs); }
+				if (ProfileGoConfig.bUncapped) { ClientRole.CommandLineParams.Add("profilego.uncapped"); }
+				if (ProfileGoConfig.bExit) { ClientRole.CommandLineParams.Add("profilego.exit"); }
+				if (ProfileGoConfig.Settle > 0.0f) { ClientRole.CommandLineParams.Add("profilego.settle", ProfileGoConfig.Settle); }
+				if (ProfileGoConfig.bDisablePlayerInput) { ClientRole.CommandLineParams.Add("profilego.disableplayerinput"); }
+			}
+		}
+
+		protected override bool CopyCSVs(List<FileInfo> CsvFiles, DirectoryInfo CSVDirectory)
+		{
+			// For ProfileGo, each "scenario" can generate a CSV file, so we have
+			// to ensure we copy all valid CSVs to output directory
+			if(CsvFiles.Count == 0)
+			{
+				return false;
+			}
+
+			foreach (FileInfo CsvFile in CsvFiles)
+			{
+				CopyCSV(CsvFile, CSVDirectory);
+			}
+
+			return true;
+		}
+
+		private ProfileGoConfig ProfileGoConfig;
+	}
+
+	public class ProfileGoTest : AutomatedProfileGoTestNode<AutomatedProfileGoTestConfig>
+	{
+		public ProfileGoTest(UnrealTestContext InContext)
+			: base(InContext) { }
+	}
+}
diff --git a/Engine/Plugins/Performance/AutomatedPerfTesting/Build/Scripts/AutomatedReplayPerfTest.cs b/Engine/Plugins/Performance/AutomatedPerfTesting/Build/Scripts/AutomatedReplayPerfTest.cs
new file mode 100644
--- /dev/null
+++ b/Engine/Plugins/Performance/AutomatedPerfTesting/Build/Scripts/AutomatedReplayPerfTest.cs
@@ -0,0 +1,208 @@
+// Copyright Epic Games, Inc. All Rights Reserved.
+
+using System;
+using System.Collections.Generic;
+using System.IO;
+using System.Linq;
+using Gauntlet;
+using AutomationTool;
+using Log = Gauntlet.Log;
+using UnrealBuildTool;
+
+namespace AutomatedPerfTest
+{
+	public class AutomatedReplayPerfTestConfig : AutomatedPerfTestConfigBase
+	{
+		/// <summary>
+		/// Which replay to run the test on
+		/// </summary>
+		[AutoParamWithNames("", "AutomatedPerfTest.ReplayPerfTest.ReplayName")]
+		public string ReplayName;
+	}
+
+	/// <summary>
+	/// Implementation of a Gauntlet TestNode for AutomatedPerfTest plugin
+	/// </summary>
+	/// <typeparam name="TConfigClass"></typeparam>
+	public abstract class AutomatedReplayPerfTestNode<TConfigClass> : AutomatedPerfTestNode<TConfigClass>, IAutomatedPerfTest
+		where TConfigClass : AutomatedReplayPerfTestConfig, new()
+	{
+		public AutomatedReplayPerfTestNode(UnrealTestContext InContext) : base(InContext)
+		{
+			Config = null;
+		}
+
+		/// <summary>
+		/// Handler to explicitly copy files to the device. It is assumed this
+		/// is called when Gauntlet configures the device before starting the
+		/// tests.
+		/// </summary>
+		public void CopyReplayToDevice(ITargetDevice Device)
+		{
+			if (Config == null)
+			{
+				Log.Warning("ConfigureDevice() called before Test Node Configuration is set.");
+				return;
+			}
+
+			// At the moment in some pathways we don't seem to copy
+			// files in a role. This ensures we have a copy of the
+			// replay files on the target device regardless.
+			if (Config.GetRequiredRoles(UnrealTargetRole.Client).Any())
+			{
+				foreach (UnrealTestRole ClientRole in Config.GetRequiredRoles(UnrealTargetRole.Client))
+				{
+					UnrealTestRoleContext RoleContext = Context.GetRoleContext(UnrealTargetRole.Client);
+					UnrealTargetPlatform Platform = RoleContext.Platform;
+					if (Device.Platform == Platform && ClientRole.FilesToCopy.Any())
+					{
+						UnrealSessionRole Role = new(RoleContext.Type, Platform, RoleContext.Configuration);
+						UnrealAppConfig AppConfig = Context.BuildInfo.CreateConfiguration(Role);
+
+						// needed so AFS will work on CopyAdditional files
+						_ = Device.CreateAppInstall(AppConfig);
+
+						Device.CopyAdditionalFiles(ClientRole.FilesToCopy);
+					}
+				}
+			}
+		}
+
+		public override TConfigClass GetConfiguration()
+		{
+			Config = base.GetConfiguration();
+
+			Config.DataSourceName = Config.GetDataSourceName(Context.BuildInfo.ProjectName, ReplayDataSourceName);
+
+			IEnumerable<UnrealTestRole> ClientRoles = Config.GetRequiredRoles(UnrealTargetRole.Client);
+			// extend the role(s) that we initialized in the base class
+			if (ClientRoles.Any())
+			{
+				foreach (UnrealTestRole ClientRole in ClientRoles)
+				{
+					ClientRole.Controllers.Add("AutomatedReplayPerfTest");
+					ClientRole.FilesToCopy.Clear();
+
+					UnrealTargetPlatform Platform = Context.GetRoleContext(UnrealTargetRole.Client).Platform;
+					List<string> ReplaysFromConfig = GetTestsFromConfig();
+
+					// Replay name not provided
+					if (string.IsNullOrEmpty(Config.ReplayName))
+					{
+						if(ReplaysFromConfig == null || !ReplaysFromConfig.Any())
+						{
+							throw new AutomationException("No replays found in settings or provided via arguments.");
+						}
+
+						// Use the first replay found. We need to extract the replay path from settings as we need to
+						// post process the path depending on the target platform.
+						Config.ReplayName = ReplaysFromConfig[0];
+					}
+					else
+					{
+						foreach(string Replay in ReplaysFromConfig)
+						{
+							if(!string.IsNullOrEmpty(Replay) && Replay.Contains(Config.ReplayName))
+							{
+								Log.Info("Found replay in settings");
+								Config.ReplayName = Replay;
+								break;
+							}
+						}
+					}
+
+					string ReplayName = Path.GetFullPath(Config.ReplayName);
+					bool bFileExists = File.Exists(ReplayName);
+					if (!bFileExists)
+					{
+						// In case the replay file specified is not in settings
+						// or not an absolute path. Check under project directory
+						// to be sure.
+						ReplayName = GetPathInProject(Context, Config.ReplayName);
+						bFileExists = File.Exists(ReplayName);
+					}
+
+					if(!bFileExists)
+					{
+						throw new AutomationException($"Replay file '{Config.ReplayName}' not found");
+					}
+
+					if (PlatformTargetSupport.IsHostMountingSupported(Platform))
+					{
+						// If current device supports host mounting of files, update the path and use host mounting
+						ReplayName = PlatformTargetSupport.GetHostMountedPath(Platform, ReplayName);
+						Log.Info($"Host Mountable Platform Detected. Updated Replay File Path to: {ReplayName} ");
+					}
+					else
+					{
+						// Copy replay file to Demos folder in device if host mounting is not supported
+						Log.Info($"Copy Replay File Path to Device: {ReplayName}");
+						UnrealFileToCopy FileToCopy = new UnrealFileToCopy(ReplayName, EIntendedBaseCopyDirectory.Demos, Path.GetFileName(ReplayName));
+						ClientRole.FilesToCopy.Add(FileToCopy);
+						ClientRole.ConfigureDevice = CopyReplayToDevice;
+
+						// If we copy to the "Demos" folder, we just need to pass the Replay Name without
+						// path and extension as the replay subsystem will automatically pick up the file
+						// from this folder.
+						ReplayName = Path.GetFileNameWithoutExtension(ReplayName);
+					}
+
+					ClientRole.CommandLineParams.AddUnique(ReplayParamName, ReplayName);
+					Log.Info($"{ReplayParamName}=\"{ReplayName}\"");
+				}
+			}
+
+			return Config;
+		}
+
+		public List<string> GetTestsFromConfig()
+		{
+			List<string> OutReplayList = new List<string>();
+			ReadConfigArray(Context,
+				"/Script/AutomatedPerfTesting.AutomatedReplayPerfTestProjectSettings",
+				"ReplaysToTest",
+				Config =>
+				{
+					Dictionary<string, string> ReplayConfig = IniConfigUtil.ParseDictionaryFromConfigString(Config);
+					string Path;
+					if (ReplayConfig.TryGetValue("FilePath", out Path))
+					{
+						OutReplayList.Add(GetPathInProject(Context, Path.Replace("\"", "")));
+					}
+				});
+
+			return OutReplayList;
+		}
+
+		protected override string GetNormalizedInsightsFileName(string CSVFileName)
+		{
+			return GetNormalizedInsightsFileName(CSVFileName, "_Replay");
+		}
+
+		/// <inheritdoc/>
+		protected override string CreateTestIdentity(TConfigClass Config = null)
+		{
+			Config ??= GetCachedConfiguration();
+
+			string rootTestIdentity = base.CreateTestIdentity(Config);
+			string replayName = string.IsNullOrEmpty(Config.ReplayName) ? string.Empty : Path.GetFileNameWithoutExtension(Config.ReplayName);
+
+			return string.IsNullOrEmpty(replayName) ? rootTestIdentity : $"{rootTestIdentity}.{replayName}";
+		}
+
+		private readonly string ReplayDataSourceName = "ReplayRun";
+		private readonly string ReplayParamName = "AutomatedPerfTest.ReplayPerfTest.ReplayName";
+		private TConfigClass Config;
+	}
+
+	/// <summary>
+	/// "Standard issue" implementation usable for samples that don't need anything more advanced
+	/// </summary>
+	public class ReplayTest : AutomatedReplayPerfTestNode<AutomatedReplayPerfTestConfig>
+	{
+		public ReplayTest(Gauntlet.UnrealTestContext InContext)
+			: base(InContext)
+		{
+		}
+	}
+}
diff --git a/Engine/Plugins/Performance/AutomatedPerfTesting/Build/Scripts/AutomatedSequencePerfTest.cs b/Engine/Plugins/Performance/AutomatedPerfTesting/Build/Scripts/AutomatedSequencePerfTest.cs
new file mode 100644
--- /dev/null
+++ b/Engine/Plugins/Performance/AutomatedPerfTesting/Build/Scripts/AutomatedSequencePerfTest.cs
@@ -0,0 +1,103 @@
+// Copyright Epic Games, Inc. All Rights Reserved.
+
+using System;
+using System.Collections.Generic;
+using System.Linq;
+using Gauntlet;
+using AutomationTool;
+
+namespace AutomatedPerfTest
+{
+	public class AutomatedSequencePerfTestConfig : AutomatedPerfTestConfigBase
+	{
+		/// <summary>
+		/// Which map to run the test on
+		/// </summary>
+		[AutoParamWithNames("", "AutomatedPerfTest.SequencePerfTest.MapSequenceName")]
+		public string MapSequenceComboName;
+	}
+
+	/// <summary>
+	/// Implementation of a Gauntlet TestNode for AutomatedPerfTest plugin
+	/// </summary>
+	/// <typeparam name="TConfigClass"></typeparam>
+	public abstract class AutomatedSequencePerfTestNode<TConfigClass> : AutomatedPerfTestNode<TConfigClass>, IAutomatedPerfTest
+		where TConfigClass : AutomatedSequencePerfTestConfig, new()
+	{
+		public AutomatedSequencePerfTestNode(UnrealTestContext InContext) : base(InContext)
+		{
+			SummaryTable = "sequence";
+		}
+
+		public override TConfigClass GetConfiguration()
+		{
+			TConfigClass Config = base.GetConfiguration();
+
+			Config.DataSourceName = Config.GetDataSourceName(Context.BuildInfo.ProjectName, "Sequence");
+
+			// extend the role(s) that we initialized in the base class
+			if (Config.GetRequiredRoles(UnrealTargetRole.Client).Any())
+			{
+				foreach(UnrealTestRole ClientRole in Config.GetRequiredRoles(UnrealTargetRole.Client))
+				{
+					ClientRole.Controllers.Add("AutomatedSequencePerfTest");
+
+					// if a specific MapSequenceComboName was defined in the commandline to UAT, then add that to the commandline for the role
+					if (!string.IsNullOrEmpty(Config.MapSequenceComboName))
+					{
+						// use add Unique, since there should only ever be one of these specified
+						ClientRole.CommandLineParams.AddUnique($"AutomatedPerfTest.SequencePerfTest.MapSequenceName",
+							Config.MapSequenceComboName);
+					}
+				}
+			}
+
+			return Config;
+		}
+
+		public List<string> GetTestsFromConfig()
+		{
+			List<string> OutSequenceList = new List<string>();
+			ReadConfigArray(Context,
+				"/Script/AutomatedPerfTesting.AutomatedSequencePerfTestProjectSettings",
+				"MapsAndSequencesToTest",
+				Config =>
+				{
+					Dictionary<string, string> SequenceConfig = IniConfigUtil.ParseDictionaryFromConfigString(Config);
+					string ComboName;
+					if (SequenceConfig.TryGetValue("ComboName", out ComboName))
+					{
+						OutSequenceList.Add(ComboName.Replace("\"", ""));
+					}
+				});
+
+			return OutSequenceList;
+		}
+
+		protected override string GetNormalizedInsightsFileName(string CSVFileName)
+		{
+			return GetNormalizedInsightsFileName(CSVFileName, "_Sequence");
+		}
+
+		/// <inheritdoc/>
+		protected override string CreateTestIdentity(TConfigClass Config = null)
+		{
+			Config ??= GetCachedConfiguration();
+
+			string rootTestIdentity = base.CreateTestIdentity(Config);
+
+			return $"{rootTestIdentity}.{Config.MapSequenceComboName}";
+		}
+	}
+
+	/// <summary>
+	/// "Standard issue" implementation usable for samples that don't need anything more advanced
+	/// </summary>
+	public class SequenceTest : AutomatedSequencePerfTestNode<AutomatedSequencePerfTestConfig>
+	{
+		public SequenceTest(Gauntlet.UnrealTestContext InContext)
+			: base(InContext)
+		{
+		}
+	}
+}
diff --git a/Engine/Plugins/Performance/AutomatedPerfTesting/Build/Scripts/AutomatedStaticCameraPerfTest.cs b/Engine/Plugins/Performance/AutomatedPerfTesting/Build/Scripts/AutomatedStaticCameraPerfTest.cs
new file mode 100644
--- /dev/null
+++ b/Engine/Plugins/Performance/AutomatedPerfTesting/Build/Scripts/AutomatedStaticCameraPerfTest.cs
@@ -0,0 +1,83 @@
+// Copyright Epic Games, Inc. All Rights Reserved.
+
+using System.Collections.Generic;
+using System.Linq;
+using Gauntlet;
+
+namespace AutomatedPerfTest
+{
+	public class AutomatedStaticCameraPerfTestConfig : AutomatedPerfTestConfigBase
+	{
+		/// <summary>
+		/// Which map to run the test on
+		/// </summary>
+		[AutoParamWithNames("", "AutomatedPerfTest.StaticCameraPerfTest.MapName")]
+		public string MapName;
+	}
+
+	/// <summary>
+	/// Implementation of a Gauntlet TestNode for AutomatedPerfTest plugin
+	/// </summary>
+	/// <typeparam name="TConfigClass"></typeparam>
+	public abstract class AutomatedStaticCameraPerfTestNode<TConfigClass> : AutomatedPerfTestNode<TConfigClass>, IAutomatedPerfTest
+		where TConfigClass : AutomatedStaticCameraPerfTestConfig, new()
+	{
+		public AutomatedStaticCameraPerfTestNode(UnrealTestContext InContext) : base(InContext)
+		{
+			SummaryTable = "staticcamera";
+		}
+
+		public List<string> GetTestsFromConfig()
+		{
+			List<string> OutMaps = new List<string>();
+			ReadConfigArray(Context,
+				"/Script/AutomatedPerfTesting.AutomatedStaticCameraPerfTestProjectSettings",
+				"MapsToTest",
+				Map => OutMaps.Add(Map.Replace("\"", "")));
+
+			return OutMaps;
+		}
+
+		public override TConfigClass GetConfiguration()
+		{
+			TConfigClass Config = base.GetConfiguration();
+
+			Config.DataSourceName = Config.GetDataSourceName(Context.BuildInfo.ProjectName, "StaticCamera");
+
+			// extend the role(s) that we initialized in the base class
+			if (Config.GetRequiredRoles(UnrealTargetRole.Client).Any())
+			{
+				foreach(UnrealTestRole ClientRole in Config.GetRequiredRoles(UnrealTargetRole.Client))
+				{
+					ClientRole.Controllers.Add("AutomatedPlacedStaticCameraPerfTest");
+
+					// if a specific MapName was defined in the commandline to UAT, then add that to the commandline for the role
+					if (!string.IsNullOrEmpty(Config.MapName))
+					{
+						// use add Unique, since there should only ever be one of these specified
+						ClientRole.CommandLineParams.AddUnique($"AutomatedPerfTest.StaticCameraPerfTest.MapName",
+							Config.MapName);
+					}
+				}
+			}
+
+			return Config;
+		}
+
+		protected override string GetNormalizedInsightsFileName(string CSVFileName)
+		{
+			return GetNormalizedInsightsFileName(CSVFileName, "_StaticCamera");
+		}
+	}
+
+	/// <summary>
+	/// "Standard issue" implementation usable for samples that don't need anything more advanced
+	/// </summary>
+	public class StaticCameraTest : AutomatedStaticCameraPerfTestNode<AutomatedStaticCameraPerfTestConfig>
+	{
+		public StaticCameraTest(Gauntlet.UnrealTestContext InContext)
+			: base(InContext)
+		{
+		}
+	}
+}
